Animation: Started moving animation nodes out of Engine into a new AnimGraphRuntime module
- Moved all skeletal control nodes to start out

Upgrade Note: Modules containing custom animation nodes will need to have a dependency on the new "AnimGraphRuntime" module added to their Build.cs file

diff --git a/Engine/Config/BaseEngine.ini b/Engine/Config/BaseEngine.ini
--- a/Engine/Config/BaseEngine.ini
+++ b/Engine/Config/BaseEngine.ini
@@ -270,13 +270,6 @@ bUseStreamingPause=false
 +ActiveClassRedirects=(OldClassName="DynamicPhysicsVolume",NewClassName="PhysicsVolume")
 +ActiveClassRedirects=(OldClassName="DynamicTriggerVolume",NewClassName="TriggerVolume")
 
-+ActiveClassRedirects=(OldClassName="AnimNode_SkeletalControlBase",NewClassName="/Script/Engine.AnimNode_SkeletalControlBase")
-+ActiveClassRedirects=(OldClassName="AnimNode_TwoBoneIK",NewClassName="/Script/Engine.AnimNode_TwoBoneIK")
-+ActiveClassRedirects=(OldClassName="AnimNode_RotationMultiplier",NewClassName="/Script/Engine.AnimNode_RotationMultiplier")
-+ActiveClassRedirects=(OldClassName="AnimNode_ModifyBone",NewClassName="/Script/Engine.AnimNode_ModifyBone")
-+ActiveClassRedirects=(OldClassName="AnimNode_CopyBone",NewClassName="/Script/Engine.AnimNode_CopyBone")
-+ActiveClassRedirects=(OldClassName="AnimNode_SpringBone",NewClassName="/Script/Engine.AnimNode_SpringBone")
-
 +EnumRedirects=(EnumName="EConstraintTransform",OldEnumEntry="EConstraintTransform::Absoluate",NewEnumEntry="EConstraintTransform::Absolute")
 
 ;Navigation redirects
@@ -375,6 +368,24 @@ bUseStreamingPause=false
 +ActiveClassRedirects=(OldClassName="AnimationCompressionAlgorithm_RemoveTrivialKeys",NewClassName="AnimCompress_RemoveTrivialKeys")
 +ActiveClassRedirects=(OldClassName="AnimationCompressionAlgorithm_RevertToRaw",NewClassName="AnimCompress_RevertToRaw")
 
+; Moving animation nodes from Engine to AnimGraphRuntime (4.9 era)
++ActiveClassRedirects=(OldClassName="ComponentType",NewClassName="/Script/AnimGraphRuntime.ComponentType")
++ActiveClassRedirects=(OldClassName="AnimNode_BoneDrivenController",NewClassName="/Script/AnimGraphRuntime.AnimNode_BoneDrivenController")
++ActiveClassRedirects=(OldClassName="AnimNode_CopyBone",NewClassName="/Script/AnimGraphRuntime.AnimNode_CopyBone")
++ActiveClassRedirects=(OldClassName="AnimNode_HandIKRetargeting",NewClassName="/Script/AnimGraphRuntime.AnimNode_HandIKRetargeting")
++ActiveClassRedirects=(OldClassName="AxisOption",NewClassName="/Script/AnimGraphRuntime.AxisOption")
++ActiveClassRedirects=(OldClassName="InterpolationBlend",NewClassName="/Script/AnimGraphRuntime.InterpolationBlend")
++ActiveClassRedirects=(OldClassName="AnimNode_LookAt",NewClassName="/Script/AnimGraphRuntime.AnimNode_LookAt")
++ActiveClassRedirects=(OldClassName="BoneModificationMode",NewClassName="/Script/AnimGraphRuntime.BoneModificationMode")
++ActiveClassRedirects=(OldClassName="AnimNode_ModifyBone",NewClassName="/Script/AnimGraphRuntime.AnimNode_ModifyBone")
++ActiveClassRedirects=(OldClassName="BoneAxis",NewClassName="/Script/AnimGraphRuntime.BoneAxis")
++ActiveClassRedirects=(OldClassName="AnimNode_RotationMultiplier",NewClassName="/Script/AnimGraphRuntime.AnimNode_RotationMultiplier")
++ActiveClassRedirects=(OldClassName="AnimNode_SkeletalControlBase",NewClassName="/Script/AnimGraphRuntime.AnimNode_SkeletalControlBase")
++ActiveClassRedirects=(OldClassName="AnimNode_SpringBone",NewClassName="/Script/AnimGraphRuntime.AnimNode_SpringBone")
++ActiveClassRedirects=(OldClassName="AnimNode_Trail",NewClassName="/Script/AnimGraphRuntime.AnimNode_Trail")
++ActiveClassRedirects=(OldClassName="AnimNode_TwoBoneIK",NewClassName="/Script/AnimGraphRuntime.AnimNode_TwoBoneIK")
++ActiveClassRedirects=(OldClassName="AnimNode_WheelHandler",NewClassName="/Script/AnimGraphRuntime.AnimNode_WheelHandler")
+
 ;AIModule redirections
 +ActiveClassRedirects=(OldClassName="AIController",NewClassName="/Script/AIModule.AIController")
 +ActiveClassRedirects=(OldClassName="AIResourceInterface",NewClassName="/Script/AIModule.AIResourceInterface")
diff --git a/Engine/Source/Editor/AnimGraph/AnimGraph.Build.cs b/Engine/Source/Editor/AnimGraph/AnimGraph.Build.cs
--- a/Engine/Source/Editor/AnimGraph/AnimGraph.Build.cs
+++ b/Engine/Source/Editor/AnimGraph/AnimGraph.Build.cs
@@ -15,6 +15,7 @@ public class AnimGraph : ModuleRules
 				"CoreUObject", 
 				"Engine", 
 				"Slate",
+				"AnimGraphRuntime",
 				"BlueprintGraph",
 			}
 		);
diff --git a/Engine/Source/Editor/AnimGraph/Classes/AnimGraphNode_BoneDrivenController.h b/Engine/Source/Editor/AnimGraph/Classes/AnimGraphNode_BoneDrivenController.h
--- a/Engine/Source/Editor/AnimGraph/Classes/AnimGraphNode_BoneDrivenController.h
+++ b/Engine/Source/Editor/AnimGraph/Classes/AnimGraphNode_BoneDrivenController.h
@@ -3,7 +3,7 @@
 #pragma once
 
 #include "AnimGraphNode_SkeletalControlBase.h"
-#include "Animation/BoneControllers/AnimNode_BoneDrivenController.h"
+#include "BoneControllers/AnimNode_BoneDrivenController.h"
 #include "EdGraph/EdGraphNodeUtils.h" // for FNodeTitleTextTable
 #include "AnimGraphNode_BoneDrivenController.generated.h"
 
diff --git a/Engine/Source/Editor/AnimGraph/Classes/AnimGraphNode_CopyBone.h b/Engine/Source/Editor/AnimGraph/Classes/AnimGraphNode_CopyBone.h
--- a/Engine/Source/Editor/AnimGraph/Classes/AnimGraphNode_CopyBone.h
+++ b/Engine/Source/Editor/AnimGraph/Classes/AnimGraphNode_CopyBone.h
@@ -2,7 +2,7 @@
 
 #pragma once
 #include "AnimGraphNode_SkeletalControlBase.h"
-#include "Animation/BoneControllers/AnimNode_CopyBone.h"
+#include "BoneControllers/AnimNode_CopyBone.h"
 #include "EdGraph/EdGraphNodeUtils.h" // for FNodeTitleTextTable
 #include "AnimGraphNode_CopyBone.generated.h"
 
diff --git a/Engine/Source/Editor/AnimGraph/Classes/AnimGraphNode_Fabrik.h b/Engine/Source/Editor/AnimGraph/Classes/AnimGraphNode_Fabrik.h
--- a/Engine/Source/Editor/AnimGraph/Classes/AnimGraphNode_Fabrik.h
+++ b/Engine/Source/Editor/AnimGraph/Classes/AnimGraphNode_Fabrik.h
@@ -2,7 +2,7 @@
 
 #pragma once
 #include "AnimGraphNode_SkeletalControlBase.h"
-#include "Animation/BoneControllers/AnimNode_Fabrik.h"
+#include "BoneControllers/AnimNode_Fabrik.h"
 #include "AnimGraphNode_Fabrik.generated.h"
 
 // Editor node for FABRIK IK skeletal controller
diff --git a/Engine/Source/Editor/AnimGraph/Classes/AnimGraphNode_HandIKRetargeting.h b/Engine/Source/Editor/AnimGraph/Classes/AnimGraphNode_HandIKRetargeting.h
--- a/Engine/Source/Editor/AnimGraph/Classes/AnimGraphNode_HandIKRetargeting.h
+++ b/Engine/Source/Editor/AnimGraph/Classes/AnimGraphNode_HandIKRetargeting.h
@@ -2,7 +2,7 @@
 
 #pragma once
 #include "AnimGraphNode_SkeletalControlBase.h"
-#include "Animation/BoneControllers/AnimNode_HandIKRetargeting.h"
+#include "BoneControllers/AnimNode_HandIKRetargeting.h"
 #include "AnimGraphNode_HandIKRetargeting.generated.h"
 
 UCLASS(MinimalAPI)
diff --git a/Engine/Source/Editor/AnimGraph/Classes/AnimGraphNode_LookAt.h b/Engine/Source/Editor/AnimGraph/Classes/AnimGraphNode_LookAt.h
--- a/Engine/Source/Editor/AnimGraph/Classes/AnimGraphNode_LookAt.h
+++ b/Engine/Source/Editor/AnimGraph/Classes/AnimGraphNode_LookAt.h
@@ -3,7 +3,7 @@
 #pragma once
 
 #include "AnimGraphNode_SkeletalControlBase.h"
-#include "Animation/BoneControllers/AnimNode_LookAt.h"
+#include "BoneControllers/AnimNode_LookAt.h"
 #include "EdGraph/EdGraphNodeUtils.h" // for FNodeTitleTextTable
 #include "AnimGraphNode_LookAt.generated.h"
 
diff --git a/Engine/Source/Editor/AnimGraph/Classes/AnimGraphNode_ModifyBone.h b/Engine/Source/Editor/AnimGraph/Classes/AnimGraphNode_ModifyBone.h
--- a/Engine/Source/Editor/AnimGraph/Classes/AnimGraphNode_ModifyBone.h
+++ b/Engine/Source/Editor/AnimGraph/Classes/AnimGraphNode_ModifyBone.h
@@ -3,7 +3,7 @@
 #pragma once
 
 #include "AnimGraphNode_SkeletalControlBase.h"
-#include "Animation/BoneControllers/AnimNode_ModifyBone.h"
+#include "BoneControllers/AnimNode_ModifyBone.h"
 #include "EdGraph/EdGraphNodeUtils.h" // for FNodeTitleTextTable
 #include "AnimGraphNode_ModifyBone.generated.h"
 
diff --git a/Engine/Source/Editor/AnimGraph/Classes/AnimGraphNode_RotationMultiplier.h b/Engine/Source/Editor/AnimGraph/Classes/AnimGraphNode_RotationMultiplier.h
--- a/Engine/Source/Editor/AnimGraph/Classes/AnimGraphNode_RotationMultiplier.h
+++ b/Engine/Source/Editor/AnimGraph/Classes/AnimGraphNode_RotationMultiplier.h
@@ -2,7 +2,7 @@
 
 #pragma once
 #include "AnimGraphNode_SkeletalControlBase.h"
-#include "Animation/BoneControllers/AnimNode_RotationMultiplier.h"
+#include "BoneControllers/AnimNode_RotationMultiplier.h"
 #include "EdGraph/EdGraphNodeUtils.h" // for FNodeTitleTextTable
 #include "AnimGraphNode_RotationMultiplier.generated.h"
 
diff --git a/Engine/Source/Editor/AnimGraph/Classes/AnimGraphNode_SpringBone.h b/Engine/Source/Editor/AnimGraph/Classes/AnimGraphNode_SpringBone.h
--- a/Engine/Source/Editor/AnimGraph/Classes/AnimGraphNode_SpringBone.h
+++ b/Engine/Source/Editor/AnimGraph/Classes/AnimGraphNode_SpringBone.h
@@ -2,7 +2,7 @@
 
 #pragma once
 #include "AnimGraphNode_SkeletalControlBase.h"
-#include "Animation/BoneControllers/AnimNode_SpringBone.h"
+#include "BoneControllers/AnimNode_SpringBone.h"
 #include "EdGraph/EdGraphNodeUtils.h" // for FNodeTitleTextTable
 #include "AnimGraphNode_SpringBone.generated.h"
 
diff --git a/Engine/Source/Editor/AnimGraph/Classes/AnimGraphNode_Trail.h b/Engine/Source/Editor/AnimGraph/Classes/AnimGraphNode_Trail.h
--- a/Engine/Source/Editor/AnimGraph/Classes/AnimGraphNode_Trail.h
+++ b/Engine/Source/Editor/AnimGraph/Classes/AnimGraphNode_Trail.h
@@ -2,7 +2,7 @@
 
 #pragma once
 #include "AnimGraphNode_SkeletalControlBase.h"
-#include "Animation/BoneControllers/AnimNode_Trail.h"
+#include "BoneControllers/AnimNode_Trail.h"
 #include "EdGraph/EdGraphNodeUtils.h" // for FNodeTitleTextTable
 #include "AnimGraphNode_Trail.generated.h"
 
diff --git a/Engine/Source/Editor/AnimGraph/Classes/AnimGraphNode_TwoBoneIK.h b/Engine/Source/Editor/AnimGraph/Classes/AnimGraphNode_TwoBoneIK.h
--- a/Engine/Source/Editor/AnimGraph/Classes/AnimGraphNode_TwoBoneIK.h
+++ b/Engine/Source/Editor/AnimGraph/Classes/AnimGraphNode_TwoBoneIK.h
@@ -2,7 +2,7 @@
 
 #pragma once
 #include "AnimGraphNode_SkeletalControlBase.h"
-#include "Animation/BoneControllers/AnimNode_TwoBoneIK.h"
+#include "BoneControllers/AnimNode_TwoBoneIK.h"
 #include "EdGraph/EdGraphNodeUtils.h" // for FNodeTitleTextTable
 #include "AnimGraphNode_TwoBoneIK.generated.h"
 
diff --git a/Engine/Source/Editor/AnimGraph/Classes/AnimGraphNode_WheelHandler.h b/Engine/Source/Editor/AnimGraph/Classes/AnimGraphNode_WheelHandler.h
--- a/Engine/Source/Editor/AnimGraph/Classes/AnimGraphNode_WheelHandler.h
+++ b/Engine/Source/Editor/AnimGraph/Classes/AnimGraphNode_WheelHandler.h
@@ -3,7 +3,7 @@
 #pragma once
 
 #include "AnimGraphNode_SkeletalControlBase.h"
-#include "Animation/BoneControllers/AnimNode_WheelHandler.h"
+#include "BoneControllers/AnimNode_WheelHandler.h"
 #include "AnimGraphNode_WheelHandler.generated.h"
 
 UCLASS(MinimalAPI, meta=(Keywords = "Modify Wheel Vehicle"))
diff --git a/Engine/Source/Editor/AnimGraph/Classes/AnimPreviewInstance.h b/Engine/Source/Editor/AnimGraph/Classes/AnimPreviewInstance.h
--- a/Engine/Source/Editor/AnimGraph/Classes/AnimPreviewInstance.h
+++ b/Engine/Source/Editor/AnimGraph/Classes/AnimPreviewInstance.h
@@ -5,7 +5,7 @@
 #include "BlueprintGraphDefinitions.h"
 #include "AnimGraphNode_ModifyBone.h"
 #include "Animation/AnimSingleNodeInstance.h"
-#include "Animation/BoneControllers/AnimNode_ModifyBone.h"
+#include "BoneControllers/AnimNode_ModifyBone.h"
 #include "AnimPreviewInstance.generated.h"
 
 /** Enum to know how montage is being played */
diff --git a/Engine/Source/Runtime/AnimGraphRuntime/Private/AnimGraphRuntimeModule.cpp b/Engine/Source/Runtime/AnimGraphRuntime/Private/AnimGraphRuntimeModule.cpp
new file mode 100644
--- /dev/null
+++ b/Engine/Source/Runtime/AnimGraphRuntime/Private/AnimGraphRuntimeModule.cpp
@@ -0,0 +1,23 @@
+// Copyright 1998-2015 Epic Games, Inc. All Rights Reserved.
+
+#include "AnimGraphRuntimePrivatePCH.h"
+//#include "Paper2DModule.h"
+
+//////////////////////////////////////////////////////////////////////////
+// FAnimGraphRuntimeModule
+
+class FAnimGraphRuntimeModule : public IModuleInterface
+{
+public:
+	virtual void StartupModule() override
+	{
+	}
+
+	virtual void ShutdownModule() override
+	{
+	}
+};
+
+//////////////////////////////////////////////////////////////////////////
+
+IMPLEMENT_MODULE(FAnimGraphRuntimeModule, AnimGraphRuntime);
diff --git a/Engine/Source/Runtime/AnimGraphRuntime/Private/AnimGraphRuntimePrivatePCH.h b/Engine/Source/Runtime/AnimGraphRuntime/Private/AnimGraphRuntimePrivatePCH.h
new file mode 100644
--- /dev/null
+++ b/Engine/Source/Runtime/AnimGraphRuntime/Private/AnimGraphRuntimePrivatePCH.h
@@ -0,0 +1,10 @@
+// Copyright 1998-2015 Epic Games, Inc. All Rights Reserved.
+
+#pragma once
+
+#include "Engine.h"
+
+// DECLARE_LOG_CATEGORY_EXTERN(LogAnimation, Verbose, All);
+
+// #include "Paper2DClasses.h"
+// #include "Paper2DModule.h"
\ No newline at end of file
diff --git a/Engine/Source/Runtime/AnimGraphRuntime/Private/BoneControllers/AnimNodeRotationMultiplier.cpp b/Engine/Source/Runtime/AnimGraphRuntime/Private/BoneControllers/AnimNodeRotationMultiplier.cpp
new file mode 100644
--- /dev/null
+++ b/Engine/Source/Runtime/AnimGraphRuntime/Private/BoneControllers/AnimNodeRotationMultiplier.cpp
@@ -0,0 +1,171 @@
+// Copyright 1998-2015 Epic Games, Inc. All Rights Reserved.
+
+#include "AnimGraphRuntimePrivatePCH.h"
+#include "BoneControllers/AnimNode_RotationMultiplier.h"
+#include "AnimationRuntime.h"
+
+/////////////////////////////////////////////////////
+// FAnimNode_RotationMultiplier
+
+FAnimNode_RotationMultiplier::FAnimNode_RotationMultiplier()
+{
+}
+
+void FAnimNode_RotationMultiplier::GatherDebugData(FNodeDebugData& DebugData)
+{
+	FString DebugLine = DebugData.GetNodeName(this);
+
+	DebugLine += "(";
+	AddDebugNodeData(DebugLine);
+	DebugLine += FString::Printf(TEXT(" Src: %s Dst: %s Multiplier: %.2f)"), *SourceBone.BoneName.ToString(), *TargetBone.BoneName.ToString(), Multiplier);
+	DebugData.AddDebugItem(DebugLine);
+
+	ComponentPose.GatherDebugData(DebugData);
+}
+
+FVector GetAxisVector(const EBoneAxis Axis)
+{
+	switch (Axis)
+	{
+	case BA_X:
+	default:
+		return FVector(1.f,0.f,0.f);
+	case BA_Y:
+		return FVector(0.f,1.f,0.f);
+	case BA_Z:
+		return FVector(0.f,0.f,1.f);
+	}
+}
+
+FQuat FAnimNode_RotationMultiplier::ExtractAngle(const FTransform& RefPoseTransform, const FTransform& LocalBoneTransform, const EBoneAxis Axis)
+{
+	// local bone transform with reference rotation
+	FTransform ReferenceBoneTransform = RefPoseTransform;
+	ReferenceBoneTransform.SetTranslation(LocalBoneTransform.GetTranslation());
+
+	// find delta angle between the two quaternions X Axis.
+	const FVector RotationAxis = GetAxisVector(Axis);
+	const FVector LocalRotationVector = LocalBoneTransform.GetRotation().RotateVector(RotationAxis);
+	const FVector ReferenceRotationVector = ReferenceBoneTransform.GetRotation().RotateVector(RotationAxis);
+
+	const FQuat LocalToRefQuat = FQuat::FindBetween(LocalRotationVector, ReferenceRotationVector);
+	checkSlow( LocalToRefQuat.IsNormalized() );
+
+	// Rotate parent bone atom from position in local space to reference skeleton
+	// Since our rotation rotates both vectors with shortest arc
+	// we're essentially left with a quaternion that has angle difference with reference skeleton version
+	const FQuat BoneQuatAligned = LocalToRefQuat* LocalBoneTransform.GetRotation();
+	checkSlow( BoneQuatAligned.IsNormalized() );
+
+	// Find that delta angle
+	const FQuat DeltaQuat = (ReferenceBoneTransform.GetRotation().Inverse()) * BoneQuatAligned;
+	checkSlow( DeltaQuat.IsNormalized() );
+
+#if 0 //DEBUG_TWISTBONECONTROLLER
+	UE_LOG(LogSkeletalControl, Log, TEXT("\t ExtractAngle, Bone: %s"), 
+		*SourceBone.BoneName.ToString());
+	UE_LOG(LogSkeletalControl, Log, TEXT("\t\t Bone Quat: %s, Rot: %s, AxisX: %s"), *LocalBoneTransform.GetRotation().ToString(), *LocalBoneTransform.GetRotation().Rotator().ToString(), *LocalRotationVector.ToString() );
+	UE_LOG(LogSkeletalControl, Log, TEXT("\t\t BoneRef Quat: %s, Rot: %s, AxisX: %s"), *ReferenceBoneTransform.GetRotation().ToString(), *ReferenceBoneTransform.GetRotation().Rotator().ToString(), *ReferenceRotationVector.ToString() );
+	UE_LOG(LogSkeletalControl, Log, TEXT("\t\t LocalToRefQuat Quat: %s, Rot: %s"), *LocalToRefQuat.ToString(), *LocalToRefQuat.Rotator().ToString() );
+
+	const FVector BoneQuatAlignedX = LocalBoneTransform.GetRotation().RotateVector(RotationAxis);
+	UE_LOG(LogSkeletalControl, Log, TEXT("\t\t BoneQuatAligned Quat: %s, Rot: %s, AxisX: %s"), *BoneQuatAligned.ToString(), *BoneQuatAligned.Rotator().ToString(), *BoneQuatAlignedX.ToString() );
+	UE_LOG(LogSkeletalControl, Log, TEXT("\t\t DeltaQuat Quat: %s, Rot: %s"), *DeltaQuat.ToString(), *DeltaQuat.Rotator().ToString() );
+
+	FTransform BoneAtomAligned(BoneQuatAligned, ReferenceBoneTransform.GetTranslation());
+	const FQuat DeltaQuatAligned = FQuat::FindBetween(BoneAtomAligned.GetScaledAxis( EAxis::X ), ReferenceBoneTransform.GetScaledAxis( EAxis::X ));
+	UE_LOG(LogSkeletalControl, Log, TEXT("\t\t DeltaQuatAligned Quat: %s, Rot: %s"), *DeltaQuatAligned.ToString(), *DeltaQuatAligned.Rotator().ToString() );
+	FVector DeltaAxis;
+	float	DeltaAngle;
+	DeltaQuat.ToAxisAndAngle(DeltaAxis, DeltaAngle);
+	UE_LOG(LogSkeletalControl, Log, TEXT("\t\t DeltaAxis: %s, DeltaAngle: %f"), *DeltaAxis.ToString(), DeltaAngle );
+#endif
+
+	return DeltaQuat;
+}
+
+FQuat FAnimNode_RotationMultiplier::MultiplyQuatBasedOnSourceIndex(const FTransform& RefPoseTransform, const FTransform& LocalBoneTransform, const EBoneAxis Axis, float InMultiplier, const FQuat& ReferenceQuat)
+{
+	// Find delta angle for source bone.
+	FQuat DeltaQuat = ExtractAngle(RefPoseTransform, LocalBoneTransform, Axis);
+
+	// Turn to Axis and Angle
+	FVector RotationAxis;
+	float	RotationAngle;
+	DeltaQuat.ToAxisAndAngle(RotationAxis, RotationAngle);
+
+	const FVector DefaultAxis = GetAxisVector(Axis);
+
+	// See if we need to invert angle - shortest path
+	if( (RotationAxis | DefaultAxis) < 0.f )
+	{
+		RotationAxis = -RotationAxis;
+		RotationAngle = -RotationAngle;
+	}
+
+	// Make sure it is the shortest angle.
+	RotationAngle = FMath::UnwindRadians(RotationAngle);
+
+	// New bone rotation
+	FQuat OutQuat = ReferenceQuat * FQuat(RotationAxis, RotationAngle* InMultiplier);
+	// Normalize resulting quaternion.
+	OutQuat.Normalize();
+
+#if 0 //DEBUG_TWISTBONECONTROLLER
+	UE_LOG(LogSkeletalControl, Log, TEXT("\t RefQuat: %s, Rot: %s"), *ReferenceQuat.ToString(), *ReferenceQuat.Rotator().ToString() );
+	UE_LOG(LogSkeletalControl, Log, TEXT("\t NewQuat: %s, Rot: %s"), *OutQuat.ToString(), *OutQuat.Rotator().ToString() );
+	UE_LOG(LogSkeletalControl, Log, TEXT("\t RollAxis: %s, RollAngle: %f"), *RotationAxis.ToString(), RotationAngle );
+#endif
+
+	return OutQuat;
+}
+
+void FAnimNode_RotationMultiplier::EvaluateBoneTransforms(USkeletalMeshComponent* SkelComp, FCSPose<FCompactPose>& MeshBases, TArray<FBoneTransform>& OutBoneTransforms)
+{
+	check(OutBoneTransforms.Num() == 0);
+
+	if ( Multiplier != 0.f )
+	{
+		// Reference bone
+		const FBoneContainer& BoneContainer = MeshBases.GetPose().GetBoneContainer();
+		const FCompactPoseBoneIndex TargetBoneIndex = TargetBone.GetCompactPoseIndex(BoneContainer);
+		const FCompactPoseBoneIndex SourceBoneIndex = SourceBone.GetCompactPoseIndex(BoneContainer);
+
+		const FQuat RefQuat = MeshBases.GetPose().GetRefPose(TargetBoneIndex).GetRotation();
+		const FTransform& SourceRefPose = MeshBases.GetPose().GetRefPose(SourceBoneIndex);
+		FQuat NewQuat = MultiplyQuatBasedOnSourceIndex(SourceRefPose, MeshBases.GetLocalSpaceTransform(SourceBoneIndex), RotationAxisToRefer, Multiplier, RefQuat);
+
+		FTransform NewLocalTransform = MeshBases.GetLocalSpaceTransform(TargetBoneIndex);
+		
+		if (bIsAdditive)
+		{
+			NewQuat = NewLocalTransform.GetRotation() * NewQuat;
+		}
+		
+		NewLocalTransform.SetRotation(NewQuat);
+
+		const FCompactPoseBoneIndex ParentIndex = MeshBases.GetPose().GetParentBoneIndex(TargetBoneIndex);
+		if( ParentIndex != INDEX_NONE )
+		{
+			const FTransform& ParentTM = MeshBases.GetComponentSpaceTransform(ParentIndex);
+			FTransform NewTransform = NewLocalTransform * ParentTM;
+			OutBoneTransforms.Add( FBoneTransform(TargetBoneIndex, NewTransform) );
+		}
+		else
+		{
+			OutBoneTransforms.Add( FBoneTransform(TargetBoneIndex, NewLocalTransform) );
+		}
+	}
+}
+
+bool FAnimNode_RotationMultiplier::IsValidToEvaluate(const USkeleton* Skeleton, const FBoneContainer& RequiredBones) 
+{
+	// if both bones are valid
+	return (TargetBone.IsValid(RequiredBones) && (TargetBone==SourceBone || SourceBone.IsValid(RequiredBones)));
+}
+
+void FAnimNode_RotationMultiplier::InitializeBoneReferences(const FBoneContainer& RequiredBones) 
+{
+	SourceBone.Initialize(RequiredBones);
+	TargetBone.Initialize(RequiredBones);
+}
\ No newline at end of file
diff --git a/Engine/Source/Runtime/AnimGraphRuntime/Private/BoneControllers/AnimNodeSkeletalControlBase.cpp b/Engine/Source/Runtime/AnimGraphRuntime/Private/BoneControllers/AnimNodeSkeletalControlBase.cpp
new file mode 100644
--- /dev/null
+++ b/Engine/Source/Runtime/AnimGraphRuntime/Private/BoneControllers/AnimNodeSkeletalControlBase.cpp
@@ -0,0 +1,74 @@
+// Copyright 1998-2015 Epic Games, Inc. All Rights Reserved.
+
+#include "AnimGraphRuntimePrivatePCH.h"
+#include "BoneControllers/AnimNode_SkeletalControlBase.h"
+#include "AnimationUtils.h"
+#include "AnimationRuntime.h"
+
+/////////////////////////////////////////////////////
+// FAnimNode_SkeletalControlBase
+
+void FAnimNode_SkeletalControlBase::Initialize(const FAnimationInitializeContext& Context)
+{
+	ComponentPose.Initialize(Context);
+}
+
+void FAnimNode_SkeletalControlBase::CacheBones(const FAnimationCacheBonesContext& Context) 
+{
+	InitializeBoneReferences(Context.AnimInstance->RequiredBones);
+	ComponentPose.CacheBones(Context);
+}
+
+void FAnimNode_SkeletalControlBase::Update(const FAnimationUpdateContext& Context)
+{
+	ComponentPose.Update(Context);
+	EvaluateGraphExposedInputs.Execute(Context);
+}
+
+bool ContainsNaN(const TArray<FBoneTransform> & BoneTransforms)
+{
+	for (int32 i = 0; i < BoneTransforms.Num(); ++i)
+	{
+		if (BoneTransforms[i].Transform.ContainsNaN())
+		{
+			return true;
+		}
+	}
+
+	return false;
+}
+
+void FAnimNode_SkeletalControlBase::EvaluateComponentSpace(FComponentSpacePoseContext& Output)
+{
+	// Evaluate the input
+	ComponentPose.EvaluateComponentSpace(Output);
+
+	// Apply the skeletal control if it's valid
+	const float ActualAlpha = AlphaScaleBias.ApplyTo(Alpha);
+	if ((ActualAlpha >= ZERO_ANIMWEIGHT_THRESH) && IsValidToEvaluate(Output.AnimInstance->CurrentSkeleton, Output.AnimInstance->RequiredBones))
+	{
+		USkeletalMeshComponent* Component = Output.AnimInstance->GetSkelMeshComponent();
+
+#if WITH_EDITORONLY_DATA
+		// save current pose before applying skeletal control to compute the exact gizmo location in AnimGraphNode
+		ForwardedPose = Output.Pose;
+#endif // #if WITH_EDITORONLY_DATA
+
+		TArray<FBoneTransform> BoneTransforms;
+		EvaluateBoneTransforms(Component, Output.Pose, BoneTransforms);
+
+		checkSlow(!ContainsNaN(BoneTransforms));
+
+		if (BoneTransforms.Num() > 0)
+		{
+			const float BlendWeight = FMath::Clamp<float>(ActualAlpha, 0.f, 1.f);
+			Output.Pose.LocalBlendCSBoneTransforms(BoneTransforms, BlendWeight);
+		}
+	}
+}
+
+void FAnimNode_SkeletalControlBase::AddDebugNodeData(FString& OutDebugData)
+{
+	const float ActualAlpha = AlphaScaleBias.ApplyTo(Alpha);
+	OutDebugData += FString::Printf(TEXT("Alpha: %.1f%%"), ActualAlpha*100.f);
+}
\ No newline at end of file
diff --git a/Engine/Source/Runtime/AnimGraphRuntime/Private/BoneControllers/AnimNode_BoneDrivenController.cpp b/Engine/Source/Runtime/AnimGraphRuntime/Private/BoneControllers/AnimNode_BoneDrivenController.cpp
new file mode 100644
--- /dev/null
+++ b/Engine/Source/Runtime/AnimGraphRuntime/Private/BoneControllers/AnimNode_BoneDrivenController.cpp
@@ -0,0 +1,132 @@
+// Copyright 1998-2015 Epic Games, Inc. All Rights Reserved.
+
+#include "AnimGraphRuntimePrivatePCH.h"
+#include "BoneControllers/AnimNode_BoneDrivenController.h"
+
+/////////////////////////////////////////////////////
+// FAnimNode_BoneDrivenController
+
+FAnimNode_BoneDrivenController::FAnimNode_BoneDrivenController()
+	: SourceComponent(EComponentType::None)
+	, TargetComponent(EComponentType::None)
+	, Multiplier(1.0f)
+	, bUseRange(false)
+	, RangeMin(-1.0f)
+	, RangeMax(1.0f)
+{
+}
+
+void FAnimNode_BoneDrivenController::GatherDebugData(FNodeDebugData& DebugData)
+{
+	FString DebugLine = DebugData.GetNodeName(this);
+	DebugLine += "(";
+	AddDebugNodeData(DebugLine);
+	DebugLine += FString::Printf(TEXT("  DrivingBone: %s\nDrivenBone: %s"), *SourceBone.BoneName.ToString(), *TargetBone.BoneName.ToString());
+	DebugData.AddDebugItem(DebugLine);
+
+	ComponentPose.GatherDebugData(DebugData);
+}
+
+void FAnimNode_BoneDrivenController::EvaluateBoneTransforms(USkeletalMeshComponent* SkelComp, FCSPose<FCompactPose>& MeshBases, TArray<FBoneTransform>& OutBoneTransforms)
+{
+	check(OutBoneTransforms.Num() == 0);
+	
+	// Early out if we're not driving from or to anything
+	if(SourceComponent == EComponentType::None || TargetComponent == EComponentType::None || Multiplier == 0.0f)
+	{
+		return;
+	}
+
+	// Get the Local space transform and the ref pose transform to see how the transform for the source bone has changed
+	const FBoneContainer& BoneContainer = MeshBases.GetPose().GetBoneContainer();
+	const FTransform& SourceOrigRef = BoneContainer.GetRefPoseArray()[SourceBone.BoneIndex];
+	FTransform SourceCurr = MeshBases.GetLocalSpaceTransform(SourceBone.GetCompactPoseIndex(BoneContainer));
+
+	// Difference from the ref pose
+	FVector RotationDiff = (SourceCurr.GetRotation() * SourceOrigRef.GetRotation().Inverse()).Euler();
+	FVector TranslationDiff = SourceCurr.GetLocation() - SourceOrigRef.GetLocation();
+	FVector CurrentScale = SourceCurr.GetScale3D();
+	FVector RefScale = SourceOrigRef.GetScale3D();
+	float ScaleDiff = FMath::Max3(CurrentScale[0], CurrentScale[1], CurrentScale[2]) - FMath::Max3(RefScale[0], RefScale[1], RefScale[2]);
+	
+	// Difference to apply after processing
+	FVector NewRotDiff(FVector::ZeroVector);
+	FVector NewTransDiff(FVector::ZeroVector);
+	float NewScaleDiff = 0.0f;
+
+	float* SourcePtr = nullptr;
+	float* DestPtr = nullptr;
+
+	// Resolve source value
+	if(SourceComponent < EComponentType::RotationX)
+	{
+		SourcePtr = &TranslationDiff[(int32)(SourceComponent - EComponentType::TranslationX)];
+	}
+	else if(SourceComponent < EComponentType::Scale)
+	{
+		SourcePtr = &RotationDiff[(int32)(SourceComponent - EComponentType::RotationX)];
+	}
+	else
+	{
+		SourcePtr = &ScaleDiff;
+	}
+	
+	// Resolve target value
+	if(TargetComponent < EComponentType::RotationX)
+	{
+		DestPtr = &NewTransDiff[(int32)(TargetComponent - EComponentType::TranslationX)];
+	}
+	else if(TargetComponent < EComponentType::Scale)
+	{
+		DestPtr = &NewRotDiff[(int32)(TargetComponent - EComponentType::RotationX)];
+	}
+	else
+	{
+		DestPtr = &NewScaleDiff;
+	}
+
+	// Set the value; clamped to the given range if necessary.
+	if(bUseRange)
+	{
+		*DestPtr = FMath::Clamp(*SourcePtr * Multiplier, RangeMin, RangeMax);
+	}
+	else
+	{
+		*DestPtr = *SourcePtr * Multiplier;
+	}
+
+	// Build final transform difference
+	FTransform FinalDiff(FQuat::MakeFromEuler(NewRotDiff), NewTransDiff, FVector(NewScaleDiff));
+
+	FCompactPoseBoneIndex TargetBoneIndex = TargetBone.GetCompactPoseIndex(BoneContainer);
+
+	// Starting point for the new transform
+	FTransform NewLocal = MeshBases.GetLocalSpaceTransform(TargetBoneIndex);
+	NewLocal.AccumulateWithAdditiveScale3D(FinalDiff);
+
+	// If we have a parent, concatenate the transform, otherwise just take the new transform
+	FCompactPoseBoneIndex PIdx = MeshBases.GetPose().GetParentBoneIndex(TargetBoneIndex);
+
+	if(PIdx != INDEX_NONE)
+	{
+		const FTransform& ParentTM = MeshBases.GetComponentSpaceTransform(PIdx);
+		
+		OutBoneTransforms.Add(FBoneTransform(TargetBoneIndex, NewLocal * ParentTM));
+	}
+	else
+	{
+		OutBoneTransforms.Add(FBoneTransform(TargetBoneIndex, NewLocal));
+	}
+
+}
+
+bool FAnimNode_BoneDrivenController::IsValidToEvaluate(const USkeleton* Skeleton, const FBoneContainer& RequiredBones)
+{
+	return SourceBone.IsValid(RequiredBones) && TargetBone.IsValid(RequiredBones);
+}
+
+void FAnimNode_BoneDrivenController::InitializeBoneReferences(const FBoneContainer& RequiredBones)
+{
+	SourceBone.Initialize(RequiredBones);
+	TargetBone.Initialize(RequiredBones);
+}
diff --git a/Engine/Source/Runtime/AnimGraphRuntime/Private/BoneControllers/AnimNode_CopyBone.cpp b/Engine/Source/Runtime/AnimGraphRuntime/Private/BoneControllers/AnimNode_CopyBone.cpp
new file mode 100644
--- /dev/null
+++ b/Engine/Source/Runtime/AnimGraphRuntime/Private/BoneControllers/AnimNode_CopyBone.cpp
@@ -0,0 +1,75 @@
+// Copyright 1998-2015 Epic Games, Inc. All Rights Reserved.
+
+#include "AnimGraphRuntimePrivatePCH.h"
+#include "BoneControllers/AnimNode_CopyBone.h"
+
+/////////////////////////////////////////////////////
+// FAnimNode_CopyBone
+
+FAnimNode_CopyBone::FAnimNode_CopyBone()
+	: bCopyTranslation(false)
+	, bCopyRotation(false)
+	, bCopyScale(false)
+{
+}
+
+void FAnimNode_CopyBone::GatherDebugData(FNodeDebugData& DebugData)
+{
+	FString DebugLine = DebugData.GetNodeName(this);
+	
+	DebugLine += "(";
+	AddDebugNodeData(DebugLine);
+	DebugLine += FString::Printf(TEXT(" Src: %s Dst: %s)"), *SourceBone.BoneName.ToString(), *TargetBone.BoneName.ToString());
+	DebugData.AddDebugItem(DebugLine);
+
+	ComponentPose.GatherDebugData(DebugData);
+}
+
+void FAnimNode_CopyBone::EvaluateBoneTransforms(USkeletalMeshComponent* SkelComp, FCSPose<FCompactPose>& MeshBases, TArray<FBoneTransform>& OutBoneTransforms)
+{
+	check(OutBoneTransforms.Num() == 0);
+
+	// Pass through if we're not doing anything.
+	if( !bCopyTranslation && !bCopyRotation && !bCopyScale )
+	{
+		return;
+	}
+
+	// Get component space transform for source and current bone.
+	const FBoneContainer& BoneContainer = MeshBases.GetPose().GetBoneContainer();
+	FCompactPoseBoneIndex TargetBoneIndex = TargetBone.GetCompactPoseIndex(BoneContainer);
+
+	const FTransform& SourceBoneTM = MeshBases.GetComponentSpaceTransform(SourceBone.GetCompactPoseIndex(BoneContainer));
+	FTransform CurrentBoneTM = MeshBases.GetComponentSpaceTransform(TargetBoneIndex);
+
+	// Copy individual components
+	if (bCopyTranslation)
+	{
+		CurrentBoneTM.SetTranslation( SourceBoneTM.GetTranslation() );
+	}
+
+	if (bCopyRotation)
+	{
+		CurrentBoneTM.SetRotation( SourceBoneTM.GetRotation() );
+	}
+
+	if (bCopyScale)
+	{
+		CurrentBoneTM.SetScale3D( SourceBoneTM.GetScale3D() );
+	}
+
+	// Output new transform for current bone.
+	OutBoneTransforms.Add(FBoneTransform(TargetBoneIndex, CurrentBoneTM));
+}
+
+bool FAnimNode_CopyBone::IsValidToEvaluate(const USkeleton* Skeleton, const FBoneContainer& RequiredBones) 
+{
+	// if both bones are valid
+	return (TargetBone.IsValid(RequiredBones) && (TargetBone==SourceBone || SourceBone.IsValid(RequiredBones)));
+}
+
+void FAnimNode_CopyBone::InitializeBoneReferences(const FBoneContainer& RequiredBones) 
+{
+	SourceBone.Initialize(RequiredBones);
+	TargetBone.Initialize(RequiredBones);
+}
\ No newline at end of file
diff --git a/Engine/Source/Runtime/AnimGraphRuntime/Private/BoneControllers/AnimNode_Fabrik.cpp b/Engine/Source/Runtime/AnimGraphRuntime/Private/BoneControllers/AnimNode_Fabrik.cpp
new file mode 100644
--- /dev/null
+++ b/Engine/Source/Runtime/AnimGraphRuntime/Private/BoneControllers/AnimNode_Fabrik.cpp
@@ -0,0 +1,257 @@
+// Copyright 1998-2015 Epic Games, Inc. All Rights Reserved.
+
+#include "AnimGraphRuntimePrivatePCH.h"
+#include "BoneControllers/AnimNode_Fabrik.h"
+
+/////////////////////////////////////////////////////
+// AnimNode_Fabrik
+// Implementation of the FABRIK IK Algorithm
+// Please see http://andreasaristidou.com/publications/FABRIK.pdf for more details
+
+FAnimNode_Fabrik::FAnimNode_Fabrik()
+	: EffectorTransform(FTransform::Identity)
+	, EffectorTransformSpace(BCS_ComponentSpace)
+	, EffectorRotationSource(BRS_KeepLocalSpaceRotation)
+	, Precision(1.f)
+	, MaxIterations(10)
+	, bEnableDebugDraw(false)
+{
+}
+
+FVector FAnimNode_Fabrik::GetCurrentLocation(FCSPose<FCompactPose>& MeshBases, const FCompactPoseBoneIndex& BoneIndex)
+{
+	return MeshBases.GetComponentSpaceTransform(BoneIndex).GetLocation();
+}
+
+void FAnimNode_Fabrik::EvaluateBoneTransforms(USkeletalMeshComponent* SkelComp, FCSPose<FCompactPose>& MeshBases, TArray<FBoneTransform>& OutBoneTransforms)
+{
+	const FBoneContainer& BoneContainer = MeshBases.GetPose().GetBoneContainer();
+
+	// Update EffectorLocation if it is based off a bone position
+	FTransform CSEffectorTransform = EffectorTransform;
+	FAnimationRuntime::ConvertBoneSpaceTransformToCS(SkelComp, MeshBases, CSEffectorTransform, EffectorTransformBone.GetCompactPoseIndex(BoneContainer), EffectorTransformSpace);
+	
+	FVector const CSEffectorLocation = CSEffectorTransform.GetLocation();
+
+	// @fixme - we need better to draw widgets and debug information in editor.
+#if !(UE_BUILD_SHIPPING || UE_BUILD_TEST)
+	if (bEnableDebugDraw)
+	{
+		// Show end effector position.
+		DrawDebugBox(SkelComp->GetWorld(), CSEffectorLocation, FVector(Precision), FColor::Green, true, 0.1f);
+		DrawDebugCoordinateSystem(SkelComp->GetWorld(), CSEffectorLocation, CSEffectorTransform.GetRotation().Rotator(), 5.f, true, 0.1f);
+	}
+#endif
+
+	// Gather all bone indices between root and tip.
+	TArray<FCompactPoseBoneIndex> BoneIndices;
+	
+	{
+		const FCompactPoseBoneIndex RootIndex = RootBone.GetCompactPoseIndex(BoneContainer);
+		FCompactPoseBoneIndex BoneIndex = TipBone.GetCompactPoseIndex(BoneContainer);
+		do
+		{
+			BoneIndices.Insert(BoneIndex, 0);
+			BoneIndex = MeshBases.GetPose().GetParentBoneIndex(BoneIndex);
+		} while (BoneIndex != RootIndex);
+		BoneIndices.Insert(BoneIndex, 0);
+	}
+
+	// Maximum length of skeleton segment at full extension
+	float MaximumReach = 0;
+
+	// Gather transforms
+	int32 const NumTransforms = BoneIndices.Num();
+	OutBoneTransforms.AddUninitialized(NumTransforms);
+
+	// Gather chain links. These are non zero length bones.
+	TArray<FABRIKChainLink> Chain;
+	Chain.Reserve(NumTransforms);
+
+	// Start with Root Bone
+	{
+		const FCompactPoseBoneIndex& RootBoneIndex = BoneIndices[0];
+		const FTransform& BoneCSTransform = MeshBases.GetComponentSpaceTransform(RootBoneIndex);
+
+		OutBoneTransforms[0] = FBoneTransform(RootBoneIndex, BoneCSTransform);
+		Chain.Add(FABRIKChainLink(BoneCSTransform.GetLocation(), 0.f, RootBoneIndex, 0));
+	}
+
+	// Go through remaining transforms
+	for (int32 TransformIndex = 1; TransformIndex < NumTransforms; TransformIndex++)
+	{
+		const FCompactPoseBoneIndex& BoneIndex = BoneIndices[TransformIndex];
+
+		const FTransform& BoneCSTransform = MeshBases.GetComponentSpaceTransform(BoneIndex);
+		FVector const BoneCSPosition = BoneCSTransform.GetLocation();
+
+		OutBoneTransforms[TransformIndex] = FBoneTransform(BoneIndex, BoneCSTransform);
+
+		// Calculate the combined length of this segment of skeleton
+		float const BoneLength = FVector::Dist(BoneCSPosition, OutBoneTransforms[TransformIndex-1].Transform.GetLocation());
+
+		if (!FMath::IsNearlyZero(BoneLength))
+		{
+			Chain.Add(FABRIKChainLink(BoneCSPosition, BoneLength, BoneIndex, TransformIndex));
+			MaximumReach += BoneLength;
+		}
+		else
+		{
+			// Mark this transform as a zero length child of the last link.
+			// It will inherit position and delta rotation from parent link.
+			FABRIKChainLink & ParentLink = Chain[Chain.Num()-1];
+			ParentLink.ChildZeroLengthTransformIndices.Add(TransformIndex);
+		}
+	}
+
+	bool bBoneLocationUpdated = false;
+	float const RootToTargetDist = FVector::Dist(Chain[0].Position, CSEffectorLocation);
+	int32 const NumChainLinks = Chain.Num();
+
+	// FABRIK algorithm - bone translation calculation
+	// If the effector is further away than the distance from root to tip, simply move all bones in a line from root to effector location
+	if (RootToTargetDist > MaximumReach)
+	{
+		for (int32 LinkIndex = 1; LinkIndex < NumChainLinks; LinkIndex++)
+		{
+			FABRIKChainLink const & ParentLink = Chain[LinkIndex - 1];
+			FABRIKChainLink & CurrentLink = Chain[LinkIndex];
+			CurrentLink.Position = ParentLink.Position + (CSEffectorLocation - ParentLink.Position).GetUnsafeNormal() * CurrentLink.Length;
+		}
+		bBoneLocationUpdated = true;
+	}
+	else // Effector is within reach, calculate bone translations to position tip at effector location
+	{
+		int32 const TipBoneLinkIndex = NumChainLinks - 1;
+
+		// Check distance between tip location and effector location
+		float Slop = FVector::Dist(Chain[TipBoneLinkIndex].Position, CSEffectorLocation);
+		if (Slop > Precision)
+		{
+			// Set tip bone at end effector location.
+			Chain[TipBoneLinkIndex].Position = CSEffectorLocation;
+
+			int32 IterationCount = 0;
+			while ((Slop > Precision) && (IterationCount++ < MaxIterations))
+			{
+				// "Forward Reaching" stage - adjust bones from end effector.
+				for (int32 LinkIndex = TipBoneLinkIndex - 1; LinkIndex > 0; LinkIndex--)
+				{
+					FABRIKChainLink & CurrentLink = Chain[LinkIndex];
+					FABRIKChainLink const & ChildLink = Chain[LinkIndex + 1];
+
+					CurrentLink.Position = ChildLink.Position + (CurrentLink.Position - ChildLink.Position).GetUnsafeNormal() * ChildLink.Length;
+				}
+
+				// "Backward Reaching" stage - adjust bones from root.
+				for (int32 LinkIndex = 1; LinkIndex < TipBoneLinkIndex; LinkIndex++)
+				{
+					FABRIKChainLink const & ParentLink = Chain[LinkIndex - 1];
+					FABRIKChainLink & CurrentLink = Chain[LinkIndex];
+
+					CurrentLink.Position = ParentLink.Position + (CurrentLink.Position - ParentLink.Position).GetUnsafeNormal() * CurrentLink.Length;
+				}
+
+				// Re-check distance between tip location and effector location
+				// Since we're keeping tip on top of effector location, check with its parent bone.
+				Slop = FMath::Abs(Chain[TipBoneLinkIndex].Length - FVector::Dist(Chain[TipBoneLinkIndex - 1].Position, CSEffectorLocation));
+			}
+
+			// Place tip bone based on how close we got to target.
+			{
+				FABRIKChainLink const & ParentLink = Chain[TipBoneLinkIndex - 1];
+				FABRIKChainLink & CurrentLink = Chain[TipBoneLinkIndex];
+
+				CurrentLink.Position = ParentLink.Position + (CurrentLink.Position - ParentLink.Position).GetUnsafeNormal() * CurrentLink.Length;
+			}
+
+			bBoneLocationUpdated = true;
+		}
+	}
+
+	// If we moved some bones, update bone transforms.
+	if (bBoneLocationUpdated)
+	{
+		// First step: update bone transform positions from chain links.
+		for (int32 LinkIndex = 0; LinkIndex < NumChainLinks; LinkIndex++)
+		{
+			FABRIKChainLink const & ChainLink = Chain[LinkIndex];
+			OutBoneTransforms[ChainLink.TransformIndex].Transform.SetTranslation(ChainLink.Position);
+
+			// If there are any zero length children, update position of those
+			int32 const NumChildren = ChainLink.ChildZeroLengthTransformIndices.Num();
+			for (int32 ChildIndex = 0; ChildIndex < NumChildren; ChildIndex++)
+			{
+				OutBoneTransforms[ChainLink.ChildZeroLengthTransformIndices[ChildIndex]].Transform.SetTranslation(ChainLink.Position);
+			}
+		}
+
+		// FABRIK algorithm - re-orientation of bone local axes after translation calculation
+		for (int32 LinkIndex = 0; LinkIndex < NumChainLinks - 1; LinkIndex++)
+		{
+			FABRIKChainLink const & CurrentLink = Chain[LinkIndex];
+			FABRIKChainLink const & ChildLink = Chain[LinkIndex + 1];
+
+			// Calculate pre-translation vector between this bone and child
+			FVector const OldDir = (GetCurrentLocation(MeshBases, ChildLink.BoneIndex) - GetCurrentLocation(MeshBases, CurrentLink.BoneIndex)).GetUnsafeNormal();
+
+			// Get vector from the post-translation bone to it's child
+			FVector const NewDir = (ChildLink.Position - CurrentLink.Position).GetUnsafeNormal();
+
+			// Calculate axis of rotation from pre-translation vector to post-translation vector
+			FVector const RotationAxis = FVector::CrossProduct(OldDir, NewDir).GetSafeNormal();
+			float const RotationAngle = FMath::Acos(FVector::DotProduct(OldDir, NewDir));
+			FQuat const DeltaRotation = FQuat(RotationAxis, RotationAngle);
+			// We're going to multiply it, in order to not have to re-normalize the final quaternion, it has to be a unit quaternion.
+			checkSlow(DeltaRotation.IsNormalized());
+
+			// Calculate absolute rotation and set it
+			FTransform& CurrentBoneTransform = OutBoneTransforms[CurrentLink.TransformIndex].Transform;
+			CurrentBoneTransform.SetRotation(DeltaRotation * CurrentBoneTransform.GetRotation());
+
+			// Update zero length children if any
+			int32 const NumChildren = CurrentLink.ChildZeroLengthTransformIndices.Num();
+			for (int32 ChildIndex = 0; ChildIndex < NumChildren; ChildIndex++)
+			{
+				FTransform& ChildBoneTransform = OutBoneTransforms[CurrentLink.ChildZeroLengthTransformIndices[ChildIndex]].Transform;
+				ChildBoneTransform.SetRotation(DeltaRotation * ChildBoneTransform.GetRotation());
+			}
+		}
+	}
+
+	// Special handling for tip bone's rotation.
+	int32 const TipBoneTransformIndex = OutBoneTransforms.Num() - 1;
+	switch (EffectorRotationSource)
+	{
+	case BRS_KeepLocalSpaceRotation:
+		OutBoneTransforms[TipBoneTransformIndex].Transform = MeshBases.GetLocalSpaceTransform(BoneIndices[TipBoneTransformIndex]) * OutBoneTransforms[TipBoneTransformIndex - 1].Transform;
+		break;
+	case BRS_CopyFromTarget:
+		OutBoneTransforms[TipBoneTransformIndex].Transform.SetRotation(CSEffectorTransform.GetRotation());
+		break;
+	case BRS_KeepComponentSpaceRotation:
+		// Don't change the orientation at all
+		break;
+	default:
+		break;
+	}
+}
+
+bool FAnimNode_Fabrik::IsValidToEvaluate(const USkeleton* Skeleton, const FBoneContainer& RequiredBones)
+{
+	// Allow evaluation if all parameters are initialized and TipBone is child of RootBone
+	return
+		(
+		TipBone.IsValid(RequiredBones)
+		&& RootBone.IsValid(RequiredBones)
+		&& Precision > 0
+		&& RequiredBones.BoneIsChildOf(TipBone.BoneIndex, RootBone.BoneIndex)
+		);
+}
+
+void FAnimNode_Fabrik::InitializeBoneReferences(const FBoneContainer& RequiredBones) 
+{
+	TipBone.Initialize(RequiredBones);
+	RootBone.Initialize(RequiredBones);
+	EffectorTransformBone.Initialize(RequiredBones);
+}
diff --git a/Engine/Source/Runtime/AnimGraphRuntime/Private/BoneControllers/AnimNode_HandIKRetargeting.cpp b/Engine/Source/Runtime/AnimGraphRuntime/Private/BoneControllers/AnimNode_HandIKRetargeting.cpp
new file mode 100644
--- /dev/null
+++ b/Engine/Source/Runtime/AnimGraphRuntime/Private/BoneControllers/AnimNode_HandIKRetargeting.cpp
@@ -0,0 +1,96 @@
+// Copyright 1998-2015 Epic Games, Inc. All Rights Reserved.
+
+#include "AnimGraphRuntimePrivatePCH.h"
+#include "BoneControllers/AnimNode_HandIKRetargeting.h"
+
+/////////////////////////////////////////////////////
+// FAnimNode_HandIKRetargeting
+
+FAnimNode_HandIKRetargeting::FAnimNode_HandIKRetargeting()
+: HandFKWeight(0.5f)
+{
+}
+
+void FAnimNode_HandIKRetargeting::GatherDebugData(FNodeDebugData& DebugData)
+{
+	FString DebugLine = DebugData.GetNodeName(this);
+
+	DebugLine += "(";
+	AddDebugNodeData(DebugLine);
+	DebugLine += FString::Printf(TEXT(" HandFKWeight: %f)"), HandFKWeight);
+	for (int32 BoneIndex = 0; BoneIndex < IKBonesToMove.Num(); BoneIndex++)
+	{
+		DebugLine += FString::Printf(TEXT(", %s)"), *IKBonesToMove[BoneIndex].BoneName.ToString());
+	}
+	DebugLine += FString::Printf(TEXT(")"));
+	DebugData.AddDebugItem(DebugLine);
+
+	ComponentPose.GatherDebugData(DebugData);
+}
+
+void FAnimNode_HandIKRetargeting::EvaluateBoneTransforms(USkeletalMeshComponent* SkelComp, FCSPose<FCompactPose>& MeshBases, TArray<FBoneTransform>& OutBoneTransforms)
+{
+	checkSlow(OutBoneTransforms.Num() == 0);
+
+	const FBoneContainer& BoneContainer = MeshBases.GetPose().GetBoneContainer();
+	// Get component space transforms for all of our IK and FK bones.
+	const FTransform& RightHandFKTM = MeshBases.GetComponentSpaceTransform(RightHandFK.GetCompactPoseIndex(BoneContainer));
+	const FTransform& LeftHandFKTM = MeshBases.GetComponentSpaceTransform(LeftHandFK.GetCompactPoseIndex(BoneContainer));
+	const FTransform& RightHandIKTM = MeshBases.GetComponentSpaceTransform(RightHandIK.GetCompactPoseIndex(BoneContainer));
+	const FTransform& LeftHandIKTM = MeshBases.GetComponentSpaceTransform(LeftHandIK.GetCompactPoseIndex(BoneContainer));
+	
+	// Compute weight FK and IK hand location. And translation from IK to FK.
+	FVector const FKLocation = FMath::Lerp<FVector>(LeftHandFKTM.GetTranslation(), RightHandFKTM.GetTranslation(), HandFKWeight);
+	FVector const IKLocation = FMath::Lerp<FVector>(LeftHandIKTM.GetTranslation(), RightHandIKTM.GetTranslation(), HandFKWeight);
+	FVector const IK_To_FK_Translation = FKLocation - IKLocation;
+
+	// If we're not translating, don't send any bones to update.
+	if (!IK_To_FK_Translation.IsNearlyZero())
+	{
+		// Move desired bones
+		for (const FBoneReference& BoneReference : IKBonesToMove)
+		{
+			if (BoneReference.IsValid(BoneContainer))
+			{
+				FCompactPoseBoneIndex BoneIndex = BoneReference.GetCompactPoseIndex(BoneContainer);
+				FTransform BoneTransform = MeshBases.GetComponentSpaceTransform(BoneIndex);
+				BoneTransform.AddToTranslation(IK_To_FK_Translation);
+
+				OutBoneTransforms.Add(FBoneTransform(BoneIndex, BoneTransform));
+			}
+		}
+	}
+}
+
+bool FAnimNode_HandIKRetargeting::IsValidToEvaluate(const USkeleton* Skeleton, const FBoneContainer& RequiredBones)
+{
+	if (RightHandFK.IsValid(RequiredBones)
+		&& LeftHandFK.IsValid(RequiredBones)
+		&& RightHandIK.IsValid(RequiredBones)
+		&& LeftHandIK.IsValid(RequiredBones))
+	{
+		// we need at least one bone to move valid.
+		for (int32 BoneIndex = 0; BoneIndex < IKBonesToMove.Num(); BoneIndex++)
+		{
+			if (IKBonesToMove[BoneIndex].IsValid(RequiredBones))
+			{
+				return true;
+			}
+		}
+	}
+
+	return false;
+}
+
+void FAnimNode_HandIKRetargeting::InitializeBoneReferences(const FBoneContainer& RequiredBones)
+{
+	RightHandFK.Initialize(RequiredBones);
+	LeftHandFK.Initialize(RequiredBones);
+	RightHandIK.Initialize(RequiredBones);
+	LeftHandIK.Initialize(RequiredBones);
+
+	for (int32 BoneIndex = 0; BoneIndex < IKBonesToMove.Num(); BoneIndex++)
+	{
+		IKBonesToMove[BoneIndex].Initialize(RequiredBones);
+	}
+}
\ No newline at end of file
diff --git a/Engine/Source/Runtime/AnimGraphRuntime/Private/BoneControllers/AnimNode_LookAt.cpp b/Engine/Source/Runtime/AnimGraphRuntime/Private/BoneControllers/AnimNode_LookAt.cpp
new file mode 100644
--- /dev/null
+++ b/Engine/Source/Runtime/AnimGraphRuntime/Private/BoneControllers/AnimNode_LookAt.cpp
@@ -0,0 +1,192 @@
+// Copyright 1998-2015 Epic Games, Inc. All Rights Reserved.
+
+#include "AnimGraphRuntimePrivatePCH.h"
+#include "BoneControllers/AnimNode_LookAt.h"
+#include "AnimationRuntime.h"
+
+/////////////////////////////////////////////////////
+// FAnimNode_LookAt
+
+FAnimNode_LookAt::FAnimNode_LookAt()
+	: LookAtLocation(FVector::ZeroVector)
+	, LookAtAxis(EAxis::X)
+	, LookAtClamp(0.f)
+	, InterpolationTime(0.f)
+	, InterpolationTriggerThreashold(0.f)
+	, CurrentLookAtLocation(FVector::ZeroVector)
+	, AccumulatedInterpoolationTime(0.f)
+{
+}
+
+void FAnimNode_LookAt::GatherDebugData(FNodeDebugData& DebugData)
+{
+	FString DebugLine = DebugData.GetNodeName(this);
+
+	DebugLine += "(";
+	AddDebugNodeData(DebugLine);
+	if (LookAtBone.BoneIndex != INDEX_NONE)
+	{
+		DebugLine += FString::Printf(TEXT(" Target: %s, Look At Bone: %s, Location : %s)"), *BoneToModify.BoneName.ToString(), *LookAtBone.BoneName.ToString(), *CurrentLookAtLocation.ToString());
+	}
+	else
+	{
+		DebugLine += FString::Printf(TEXT(" Target: %s, Look At Location : %s, Location : %s)"), *BoneToModify.BoneName.ToString(), *LookAtLocation.ToString(), *CurrentLookAtLocation.ToString());
+	}
+
+	DebugData.AddDebugItem(DebugLine);
+
+	ComponentPose.GatherDebugData(DebugData);
+}
+
+void DrawDebugData(UWorld* World, const FTransform& TransformVector, const FVector& StartLoc, const FVector& TargetLoc, FColor Color)
+{
+	FVector Start = TransformVector.TransformPosition(StartLoc);
+	FVector End = TransformVector.TransformPosition(TargetLoc);
+
+	DrawDebugLine(World, Start, End, Color);
+}
+
+void FAnimNode_LookAt::EvaluateBoneTransforms(USkeletalMeshComponent* SkelComp, FCSPose<FCompactPose>& MeshBases, TArray<FBoneTransform>& OutBoneTransforms)
+{
+	check(OutBoneTransforms.Num() == 0);
+
+	const FBoneContainer& BoneContainer = MeshBases.GetPose().GetBoneContainer();
+	const FCompactPoseBoneIndex ModifyBoneIndex = BoneToModify.GetCompactPoseIndex(BoneContainer);
+	FTransform ComponentBoneTransform = MeshBases.GetComponentSpaceTransform(ModifyBoneIndex);
+
+	// get target location
+	FVector TargetLocationInComponentSpace;
+	if (LookAtBone.IsValid(BoneContainer))
+	{
+		const FTransform& LookAtTransform  = MeshBases.GetComponentSpaceTransform(LookAtBone.GetCompactPoseIndex(BoneContainer));
+		TargetLocationInComponentSpace = LookAtTransform.GetLocation();
+	}
+	else
+	{
+		TargetLocationInComponentSpace = SkelComp->ComponentToWorld.InverseTransformPosition(LookAtLocation);
+	}
+	
+	FVector OldCurrentTargetLocation = CurrentTargetLocation;
+	FVector NewCurrentTargetLocation = TargetLocationInComponentSpace;
+
+	if ((NewCurrentTargetLocation - OldCurrentTargetLocation).SizeSquared() > InterpolationTriggerThreashold*InterpolationTriggerThreashold)
+	{
+		if (AccumulatedInterpoolationTime >= InterpolationTime)
+		{
+			// reset current Alpha, we're starting to move
+			AccumulatedInterpoolationTime = 0.f;
+		}
+
+		PreviousTargetLocation = OldCurrentTargetLocation;
+		CurrentTargetLocation = NewCurrentTargetLocation;
+	}
+	else if (InterpolationTriggerThreashold == 0.f)
+	{
+		CurrentTargetLocation = NewCurrentTargetLocation;
+	}
+
+	if (InterpolationTime > 0.f)
+	{
+		float CurrentAlpha = AccumulatedInterpoolationTime/InterpolationTime;
+
+		if (CurrentAlpha < 1.f)
+		{
+			float BlendAlpha = AlphaToBlendType(CurrentAlpha, GetInterpolationType());
+
+			CurrentLookAtLocation = FMath::Lerp(PreviousTargetLocation, CurrentTargetLocation, BlendAlpha);
+		}
+	}
+	else
+	{
+		CurrentLookAtLocation = CurrentTargetLocation;
+	}
+
+	if (bEnableDebug)
+	{
+		UWorld* World = SkelComp->GetWorld();
+
+		DrawDebugData(World, SkelComp->GetComponentToWorld(), ComponentBoneTransform.GetLocation(), PreviousTargetLocation, FColor(0, 255, 0));
+		DrawDebugData(World, SkelComp->GetComponentToWorld(), ComponentBoneTransform.GetLocation(), CurrentTargetLocation, FColor(255, 0, 0));
+		DrawDebugData(World, SkelComp->GetComponentToWorld(), ComponentBoneTransform.GetLocation(), CurrentLookAtLocation, FColor(0, 0, 255));
+	}
+
+	// lookat vector
+	FVector LookAtVector = GetAlignVector(ComponentBoneTransform, LookAtAxis);
+	// flip to target vector if it wasnt negative
+	bool bShouldFlip = LookAtAxis == EAxisOption::X_Neg || LookAtAxis == EAxisOption::Y_Neg || LookAtAxis == EAxisOption::Z_Neg;
+	FVector ToTarget = CurrentLookAtLocation - ComponentBoneTransform.GetLocation();
+	ToTarget.Normalize();
+	if (bShouldFlip)
+	{
+		ToTarget *= -1.f;
+	}
+	
+	if ( LookAtClamp > ZERO_ANIMWEIGHT_THRESH )
+	{
+		float LookAtClampInRadians = FMath::DegreesToRadians(LookAtClamp);
+		float DiffAngle = FMath::Acos(FVector::DotProduct(LookAtVector, ToTarget));
+		if (LookAtClampInRadians > 0.f && DiffAngle > LookAtClampInRadians)
+		{
+			FVector OldToTarget = ToTarget;
+			FVector DeltaTarget = ToTarget-LookAtVector;
+
+			float Ratio = LookAtClampInRadians/DiffAngle;
+			DeltaTarget *= Ratio;
+
+			ToTarget = LookAtVector + DeltaTarget;
+			ToTarget.Normalize();
+//			UE_LOG(LogAnimation, Warning, TEXT("Recalculation required - old target %f, new target %f"), 
+//				FMath::RadiansToDegrees(FMath::Acos(FVector::DotProduct(LookAtVector, OldToTarget))), FMath::RadiansToDegrees(FMath::Acos(FVector::DotProduct(LookAtVector, ToTarget))));
+		}
+	}
+
+	// get delta rotation
+	FQuat DeltaRot = FQuat::FindBetween(LookAtVector, ToTarget);
+	// transform current rotation to delta rotation
+	FQuat CurrentRot = ComponentBoneTransform.GetRotation();
+	FQuat NewRotation = DeltaRot * CurrentRot;
+	ComponentBoneTransform.SetRotation(NewRotation);
+
+	OutBoneTransforms.Add(FBoneTransform(ModifyBoneIndex, ComponentBoneTransform));
+}
+
+bool FAnimNode_LookAt::IsValidToEvaluate(const USkeleton* Skeleton, const FBoneContainer& RequiredBones) 
+{
+	// if both bones are valid
+	return (BoneToModify.IsValid(RequiredBones) && 
+		// or if name isn't set (use Look At Location) or Look at bone is valid 
+		// do not call isValid since that means if look at bone isn't in LOD, we won't evaluate
+		// we still should evaluate as long as the BoneToModify is valid even LookAtBone isn't included in required bones
+		(LookAtBone.BoneName == NAME_None || LookAtBone.BoneIndex != INDEX_NONE) );
+}
+
+void FAnimNode_LookAt::InitializeBoneReferences(const FBoneContainer& RequiredBones) 
+{
+	BoneToModify.Initialize(RequiredBones);
+	LookAtBone.Initialize(RequiredBones);
+}
+
+FVector FAnimNode_LookAt::GetAlignVector(const FTransform& Transform, EAxisOption::Type AxisOption)
+{
+	switch (AxisOption)
+	{
+	case EAxisOption::X:
+	case EAxisOption::X_Neg:
+		return Transform.GetUnitAxis(EAxis::X);
+	case EAxisOption::Y:
+	case EAxisOption::Y_Neg:
+		return Transform.GetUnitAxis(EAxis::Y);
+	case EAxisOption::Z:
+	case EAxisOption::Z_Neg:
+		return Transform.GetUnitAxis(EAxis::Z);
+	}
+
+	return FVector(1.f, 0.f, 0.f);
+}
+
+void FAnimNode_LookAt::Update(const FAnimationUpdateContext& Context)
+{
+	FAnimNode_SkeletalControlBase::Update(Context);
+
+	AccumulatedInterpoolationTime = FMath::Clamp(AccumulatedInterpoolationTime+Context.GetDeltaTime(), 0.f, InterpolationTime);;
+}
diff --git a/Engine/Source/Runtime/AnimGraphRuntime/Private/BoneControllers/AnimNode_ModifyBone.cpp b/Engine/Source/Runtime/AnimGraphRuntime/Private/BoneControllers/AnimNode_ModifyBone.cpp
new file mode 100644
--- /dev/null
+++ b/Engine/Source/Runtime/AnimGraphRuntime/Private/BoneControllers/AnimNode_ModifyBone.cpp
@@ -0,0 +1,114 @@
+// Copyright 1998-2015 Epic Games, Inc. All Rights Reserved.
+
+#include "AnimGraphRuntimePrivatePCH.h"
+#include "BoneControllers/AnimNode_ModifyBone.h"
+#include "AnimationRuntime.h"
+
+/////////////////////////////////////////////////////
+// FAnimNode_ModifyBone
+
+FAnimNode_ModifyBone::FAnimNode_ModifyBone()
+	: Translation(FVector::ZeroVector)
+	, Rotation(FRotator::ZeroRotator)
+	, Scale(FVector(1.0f))
+	, TranslationMode(BMM_Ignore)
+	, RotationMode(BMM_Ignore)
+	, ScaleMode(BMM_Ignore)
+	, TranslationSpace(BCS_ComponentSpace)
+	, RotationSpace(BCS_ComponentSpace)
+	, ScaleSpace(BCS_ComponentSpace)
+{
+}
+
+void FAnimNode_ModifyBone::GatherDebugData(FNodeDebugData& DebugData)
+{
+	FString DebugLine = DebugData.GetNodeName(this);
+
+	DebugLine += "(";
+	AddDebugNodeData(DebugLine);
+	DebugLine += FString::Printf(TEXT(" Target: %s)"), *BoneToModify.BoneName.ToString());
+	DebugData.AddDebugItem(DebugLine);
+
+	ComponentPose.GatherDebugData(DebugData);
+}
+
+void FAnimNode_ModifyBone::EvaluateBoneTransforms(USkeletalMeshComponent* SkelComp, FCSPose<FCompactPose>& MeshBases, TArray<FBoneTransform>& OutBoneTransforms)
+{
+	check(OutBoneTransforms.Num() == 0);
+
+	// the way we apply transform is same as FMatrix or FTransform
+	// we apply scale first, and rotation, and translation
+	// if you'd like to translate first, you'll need two nodes that first node does translate and second nodes to rotate.
+	const FBoneContainer BoneContainer = MeshBases.GetPose().GetBoneContainer();
+
+	FCompactPoseBoneIndex CompactPoseBoneToModify = BoneToModify.GetCompactPoseIndex(BoneContainer);
+	FTransform NewBoneTM = MeshBases.GetComponentSpaceTransform(CompactPoseBoneToModify);
+	
+	if (ScaleMode != BMM_Ignore)
+	{
+		// Convert to Bone Space.
+		FAnimationRuntime::ConvertCSTransformToBoneSpace(SkelComp, MeshBases, NewBoneTM, CompactPoseBoneToModify, ScaleSpace);
+
+		if (ScaleMode == BMM_Additive)
+		{
+			NewBoneTM.SetScale3D(NewBoneTM.GetScale3D() * Scale);
+		}
+		else
+		{
+			NewBoneTM.SetScale3D(Scale);
+		}
+
+		// Convert back to Component Space.
+		FAnimationRuntime::ConvertBoneSpaceTransformToCS(SkelComp, MeshBases, NewBoneTM, CompactPoseBoneToModify, ScaleSpace);
+	}
+
+	if (RotationMode != BMM_Ignore)
+	{
+		// Convert to Bone Space.
+		FAnimationRuntime::ConvertCSTransformToBoneSpace(SkelComp, MeshBases, NewBoneTM, CompactPoseBoneToModify, RotationSpace);
+
+		const FQuat BoneQuat(Rotation);
+		if (RotationMode == BMM_Additive)
+		{	
+			NewBoneTM.SetRotation(BoneQuat * NewBoneTM.GetRotation());
+		}
+		else
+		{
+			NewBoneTM.SetRotation(BoneQuat);
+		}
+
+		// Convert back to Component Space.
+		FAnimationRuntime::ConvertBoneSpaceTransformToCS(SkelComp, MeshBases, NewBoneTM, CompactPoseBoneToModify, RotationSpace);
+	}
+	
+	if (TranslationMode != BMM_Ignore)
+	{
+		// Convert to Bone Space.
+		FAnimationRuntime::ConvertCSTransformToBoneSpace(SkelComp, MeshBases, NewBoneTM, CompactPoseBoneToModify, TranslationSpace);
+
+		if (TranslationMode == BMM_Additive)
+		{
+			NewBoneTM.AddToTranslation(Translation);
+		}
+		else
+		{
+			NewBoneTM.SetTranslation(Translation);
+		}
+
+		// Convert back to Component Space.
+		FAnimationRuntime::ConvertBoneSpaceTransformToCS(SkelComp, MeshBases, NewBoneTM, CompactPoseBoneToModify, TranslationSpace);
+	}
+	
+	OutBoneTransforms.Add( FBoneTransform(BoneToModify.GetCompactPoseIndex(BoneContainer), NewBoneTM) );
+}
+
+bool FAnimNode_ModifyBone::IsValidToEvaluate(const USkeleton* Skeleton, const FBoneContainer& RequiredBones) 
+{
+	// if both bones are valid
+	return (BoneToModify.IsValid(RequiredBones));
+}
+
+void FAnimNode_ModifyBone::InitializeBoneReferences(const FBoneContainer& RequiredBones) 
+{
+	BoneToModify.Initialize(RequiredBones);
+}
\ No newline at end of file
diff --git a/Engine/Source/Runtime/AnimGraphRuntime/Private/BoneControllers/AnimNode_SpringBone.cpp b/Engine/Source/Runtime/AnimGraphRuntime/Private/BoneControllers/AnimNode_SpringBone.cpp
new file mode 100644
--- /dev/null
+++ b/Engine/Source/Runtime/AnimGraphRuntime/Private/BoneControllers/AnimNode_SpringBone.cpp
@@ -0,0 +1,220 @@
+// Copyright 1998-2015 Epic Games, Inc. All Rights Reserved.
+
+#include "AnimGraphRuntimePrivatePCH.h"
+#include "BoneControllers/AnimNode_SpringBone.h"
+
+/////////////////////////////////////////////////////
+// FAnimNode_SpringBone
+
+FAnimNode_SpringBone::FAnimNode_SpringBone()
+	: bLimitDisplacement(false)
+	, MaxDisplacement(0.0f)
+	, SpringStiffness(50.0f)
+	, SpringDamping(4.0f)
+	, ErrorResetThresh(256.0f)
+	, bNoZSpring_DEPRECATED(false)
+	, bTranslateX(true)
+	, bTranslateY(true)
+	, bTranslateZ(true)
+	, bRotateX(false)
+	, bRotateY(false)
+	, bRotateZ(false)
+	, RemainingTime(0.f)
+	, bHadValidStrength(false)
+	, BoneLocation(FVector::ZeroVector)
+	, BoneVelocity(FVector::ZeroVector)
+{
+}
+
+void FAnimNode_SpringBone::Initialize(const FAnimationInitializeContext& Context)
+{
+	FAnimNode_SkeletalControlBase::Initialize(Context);
+
+	RemainingTime = 0.0f;
+}
+
+void FAnimNode_SpringBone::CacheBones(const FAnimationCacheBonesContext& Context) 
+{
+	FAnimNode_SkeletalControlBase::CacheBones(Context);
+}
+
+void FAnimNode_SpringBone::Update(const FAnimationUpdateContext& Context)
+{
+	FAnimNode_SkeletalControlBase::Update(Context);
+
+	RemainingTime += Context.GetDeltaTime();
+
+	const USkeletalMeshComponent* SkelComp = Context.AnimInstance->GetSkelMeshComponent();
+	const UWorld* World = SkelComp->GetWorld();
+	check(World->GetWorldSettings());
+	// Fixed step simulation at 120hz
+	FixedTimeStep = (1.f / 120.f) * World->GetWorldSettings()->GetEffectiveTimeDilation();
+}
+
+void FAnimNode_SpringBone::GatherDebugData(FNodeDebugData& DebugData)
+{
+	const float ActualAlpha = AlphaScaleBias.ApplyTo(Alpha);
+
+	//MDW_TODO Add more output info?
+	FString DebugLine = DebugData.GetNodeName(this);
+	DebugLine += FString::Printf(TEXT("(Alpha: %.1f%% RemainingTime: %.3f)"), ActualAlpha*100.f, RemainingTime);
+
+	DebugData.AddDebugItem(DebugLine);
+	ComponentPose.GatherDebugData(DebugData);
+}
+
+FORCEINLINE void CopyToVectorByFlags(FVector& DestVec, const FVector& SrcVec, bool bX, bool bY, bool bZ)
+{
+	if (bX)
+	{
+		DestVec.X = SrcVec.X;
+	}
+	if (bY)
+	{
+		DestVec.Y = SrcVec.Y;
+	}
+	if (bZ)
+	{
+		DestVec.Z = SrcVec.Z;
+	}
+}
+
+void FAnimNode_SpringBone::EvaluateBoneTransforms(USkeletalMeshComponent* SkelComp, FCSPose<FCompactPose>& MeshBases, TArray<FBoneTransform>& OutBoneTransforms)
+{
+	check(OutBoneTransforms.Num() == 0);
+
+	const bool bNoOffset = !bTranslateX && !bTranslateY && !bTranslateZ;
+	if (bNoOffset)
+	{
+		return;
+	}
+
+	// Location of our bone in world space
+	const FBoneContainer& BoneContainer = MeshBases.GetPose().GetBoneContainer();
+
+	const FCompactPoseBoneIndex SpringBoneIndex = SpringBone.GetCompactPoseIndex(BoneContainer);
+	const FTransform& SpaceBase = MeshBases.GetComponentSpaceTransform(SpringBoneIndex);
+	FTransform BoneTransformInWorldSpace = (SkelComp != NULL) ? SpaceBase * SkelComp->GetComponentToWorld() : SpaceBase;
+
+	FVector const TargetPos = BoneTransformInWorldSpace.GetLocation();
+
+	AActor* SkelOwner = (SkelComp != NULL) ? SkelComp->GetOwner() : NULL;
+	if ((SkelComp != NULL) && (SkelComp->AttachParent != NULL) && (SkelOwner == NULL))
+	{
+		SkelOwner = SkelComp->AttachParent->GetOwner();
+	}
+
+	// Init values first time
+	if (RemainingTime == 0.0f)
+	{
+		BoneLocation = TargetPos;
+		BoneVelocity = FVector::ZeroVector;
+	}
+		
+	while (RemainingTime > FixedTimeStep)
+	{
+		// Update location of our base by how much our base moved this frame.
+		FVector const BaseTranslation = SkelOwner ? (SkelOwner->GetVelocity() * FixedTimeStep) : FVector::ZeroVector;
+		BoneLocation += BaseTranslation;
+
+		// Reinit values if outside reset threshold
+		if (((TargetPos - BoneLocation).SizeSquared() > (ErrorResetThresh*ErrorResetThresh)))
+		{
+			BoneLocation = TargetPos;
+			BoneVelocity = FVector::ZeroVector;
+		}
+
+		// Calculate error vector.
+		FVector const Error = (TargetPos - BoneLocation);
+		FVector const DampingForce = SpringDamping * BoneVelocity;
+		FVector const SpringForce = SpringStiffness * Error;
+
+		// Calculate force based on error and vel
+		FVector const Acceleration = SpringForce - DampingForce;
+
+		// Integrate velocity
+		// Make sure damping with variable frame rate actually dampens velocity. Otherwise Spring will go nuts.
+		float const CutOffDampingValue = 1.f/FixedTimeStep;
+		if (SpringDamping > CutOffDampingValue)
+		{
+			float const SafetyScale = CutOffDampingValue / SpringDamping;
+			BoneVelocity += SafetyScale * (Acceleration * FixedTimeStep);
+		}
+		else
+		{
+			BoneVelocity += (Acceleration * FixedTimeStep);
+		}
+
+		// Clamp velocity to something sane (|dX/dt| <= ErrorResetThresh)
+		float const BoneVelocityMagnitude = BoneVelocity.Size();
+		if (BoneVelocityMagnitude * FixedTimeStep > ErrorResetThresh)
+		{
+			BoneVelocity *= (ErrorResetThresh / (BoneVelocityMagnitude * FixedTimeStep));
+		}
+
+		// Integrate position
+		FVector const OldBoneLocation = BoneLocation;
+		FVector const DeltaMove = (BoneVelocity * FixedTimeStep);
+		BoneLocation += DeltaMove;
+
+		// Filter out spring translation based on our filter properties
+		CopyToVectorByFlags(BoneLocation, TargetPos, !bTranslateX, !bTranslateY, !bTranslateZ);
+
+
+		// If desired, limit error
+		if (bLimitDisplacement)
+		{
+			FVector CurrentDisp = BoneLocation - TargetPos;
+			// Too far away - project back onto sphere around target.
+			if (CurrentDisp.Size() > MaxDisplacement)
+			{
+				FVector DispDir = CurrentDisp.GetSafeNormal();
+				BoneLocation = TargetPos + (MaxDisplacement * DispDir);
+			}
+		}
+
+		// Update velocity to reflect post processing done to bone location.
+		BoneVelocity = (BoneLocation - OldBoneLocation) / FixedTimeStep;
+
+		check( !BoneLocation.ContainsNaN() );
+		check( !BoneVelocity.ContainsNaN() );
+
+		RemainingTime -= FixedTimeStep;
+	}
+
+	// Now convert back into component space and output - rotation is unchanged.
+	FTransform OutBoneTM = SpaceBase;
+	OutBoneTM.SetLocation( SkelComp->GetComponentToWorld().InverseTransformPosition(BoneLocation) );
+
+	const bool bUseRotation = bRotateX || bRotateY || bRotateZ;
+	if (bUseRotation)
+	{
+		FCompactPoseBoneIndex ParentBoneIndex = MeshBases.GetPose().GetParentBoneIndex(SpringBoneIndex);
+		const FTransform& ParentSpaceBase = MeshBases.GetComponentSpaceTransform(ParentBoneIndex);
+
+		FVector ParentToTarget = (TargetPos - ParentSpaceBase.GetLocation()).GetSafeNormal();
+		FVector ParentToCurrent = (BoneLocation - ParentSpaceBase.GetLocation()).GetSafeNormal();
+
+		FQuat AdditionalRotation = FQuat::FindBetween(ParentToTarget, ParentToCurrent);
+
+		// Filter rotation based on our filter properties
+		FVector EularRot = AdditionalRotation.Euler();
+		CopyToVectorByFlags(EularRot, FVector::ZeroVector, !bRotateX, !bRotateY, !bRotateZ);
+
+		OutBoneTM.SetRotation(FQuat::MakeFromEuler(EularRot) * OutBoneTM.GetRotation());
+	}
+
+	// Output new transform for current bone.
+	OutBoneTransforms.Add(FBoneTransform(SpringBoneIndex, OutBoneTM));
+}
+
+
+bool FAnimNode_SpringBone::IsValidToEvaluate(const USkeleton* Skeleton, const FBoneContainer& RequiredBones) 
+{
+	return (SpringBone.IsValid(RequiredBones));
+}
+
+void FAnimNode_SpringBone::InitializeBoneReferences(const FBoneContainer& RequiredBones) 
+{
+	SpringBone.Initialize(RequiredBones);
+}
diff --git a/Engine/Source/Runtime/AnimGraphRuntime/Private/BoneControllers/AnimNode_Trail.cpp b/Engine/Source/Runtime/AnimGraphRuntime/Private/BoneControllers/AnimNode_Trail.cpp
new file mode 100644
--- /dev/null
+++ b/Engine/Source/Runtime/AnimGraphRuntime/Private/BoneControllers/AnimNode_Trail.cpp
@@ -0,0 +1,234 @@
+// Copyright 1998-2015 Epic Games, Inc. All Rights Reserved.
+
+#include "AnimGraphRuntimePrivatePCH.h"
+#include "BoneControllers/AnimNode_Trail.h"
+
+/////////////////////////////////////////////////////
+// FAnimNode_Trail
+
+FAnimNode_Trail::FAnimNode_Trail()
+	: ChainLength(2)
+	, ChainBoneAxis(EAxis::X)
+	, bInvertChainBoneAxis(false)
+	, bLimitStretch(false)
+	, TrailRelaxation(10)
+	, StretchLimit(0)
+	, FakeVelocity(FVector::ZeroVector)
+	, bActorSpaceFakeVel(false)
+	, bHadValidStrength(false)
+{
+}
+
+void FAnimNode_Trail::Update(const FAnimationUpdateContext& Context)
+{
+	FAnimNode_SkeletalControlBase::Update(Context);
+
+	ThisTimstep = Context.GetDeltaTime();
+}
+
+void FAnimNode_Trail::GatherDebugData(FNodeDebugData& DebugData)
+{
+	FString DebugLine = DebugData.GetNodeName(this);
+
+	DebugLine += "(";
+	AddDebugNodeData(DebugLine);
+	DebugLine += FString::Printf(TEXT(" Active: %s)"), *TrailBone.BoneName.ToString());
+	DebugData.AddDebugItem(DebugLine);
+
+	ComponentPose.GatherDebugData(DebugData);
+}
+
+void FAnimNode_Trail::EvaluateBoneTransforms(USkeletalMeshComponent* SkelComp, FCSPose<FCompactPose>& MeshBases, TArray<FBoneTransform>& OutBoneTransforms)
+{
+	check(OutBoneTransforms.Num() == 0);
+
+	if( ChainLength < 2 )
+	{
+		return;
+	}
+
+	// The incoming BoneIndex is the 'end' of the spline chain. We need to find the 'start' by walking SplineLength bones up hierarchy.
+	// Fail if we walk past the root bone.
+
+	const FBoneContainer& BoneContainer = MeshBases.GetPose().GetBoneContainer();
+	FCompactPoseBoneIndex WalkBoneIndex = TrailBone.GetCompactPoseIndex(BoneContainer);
+
+	TArray<FCompactPoseBoneIndex> ChainBoneIndices;
+	ChainBoneIndices.AddZeroed(ChainLength);
+
+	ChainBoneIndices[ChainLength - 1] = WalkBoneIndex;
+
+	for (int32 i = 1; i < ChainLength; i++)
+	{
+		// returns to avoid a crash
+		// @TODO : shows an error message why failed
+		if (WalkBoneIndex == 0)
+		{
+			return;
+		}
+
+		// Get parent bone.
+		WalkBoneIndex = BoneContainer.GetParentBoneIndex(WalkBoneIndex);
+
+		//Insert indices at the start of array, so that parents are before children in the array.
+		int32 TransformIndex = ChainLength - (i + 1);
+		ChainBoneIndices[TransformIndex] = WalkBoneIndex;
+	}
+
+	OutBoneTransforms.AddZeroed(ChainLength);
+
+	// If we have >0 this frame, but didn't last time, record positions of all the bones.
+	// Also do this if number has changed or array is zero.
+	bool bHasValidStrength = (Alpha > 0.f);
+	if(TrailBoneLocations.Num() != ChainLength || (bHasValidStrength && !bHadValidStrength))
+	{
+		TrailBoneLocations.Empty();
+		TrailBoneLocations.AddZeroed(ChainLength);
+
+		for(int32 i=0; i<ChainBoneIndices.Num(); i++)
+		{
+			FCompactPoseBoneIndex ChildIndex = ChainBoneIndices[i];
+			const FTransform& ChainTransform = MeshBases.GetComponentSpaceTransform(ChildIndex);
+			TrailBoneLocations[i] = ChainTransform.GetTranslation();
+		}
+		OldLocalToWorld = SkelComp->GetTransformMatrix();
+	}
+	bHadValidStrength = bHasValidStrength;
+
+	// transform between last frame and now.
+	FMatrix OldToNewTM = OldLocalToWorld * SkelComp->GetTransformMatrix().InverseFast();
+
+	// Add fake velocity if present to all but root bone
+	if(!FakeVelocity.IsZero())
+	{
+		FVector FakeMovement = -FakeVelocity * ThisTimstep;
+
+		if (bActorSpaceFakeVel && SkelComp->GetOwner())
+		{
+			const FTransform BoneToWorld(SkelComp->GetOwner()->GetActorQuat(), SkelComp->GetOwner()->GetActorLocation());
+			FakeMovement = BoneToWorld.TransformVector(FakeMovement);
+		}
+
+		FakeMovement = SkelComp->GetTransformMatrix().InverseTransformVector(FakeMovement);
+		// Then add to each bone
+		for(int32 i=1; i<TrailBoneLocations.Num(); i++)
+		{
+			TrailBoneLocations[i] += FakeMovement;
+		}
+	}
+
+	// Root bone of trail is not modified.
+	FCompactPoseBoneIndex RootIndex = ChainBoneIndices[0]; 
+	const FTransform& ChainTransform = MeshBases.GetComponentSpaceTransform(RootIndex);
+	OutBoneTransforms[0] = FBoneTransform(RootIndex, ChainTransform);
+	TrailBoneLocations[0] = ChainTransform.GetTranslation();
+
+	// Starting one below head of chain, move bones.
+	for(int32 i=1; i<ChainBoneIndices.Num(); i++)
+	{
+		// Parent bone position in component space.
+		FCompactPoseBoneIndex ParentIndex = ChainBoneIndices[i - 1];
+		FVector ParentPos = TrailBoneLocations[i-1];
+		FVector ParentAnimPos = MeshBases.GetComponentSpaceTransform(ParentIndex).GetTranslation();
+
+		// Child bone position in component space.
+		FCompactPoseBoneIndex ChildIndex = ChainBoneIndices[i];
+		FVector ChildPos = OldToNewTM.TransformPosition(TrailBoneLocations[i]); // move from 'last frames component' frame to 'this frames component' frame
+		FVector ChildAnimPos = MeshBases.GetComponentSpaceTransform(ChildIndex).GetTranslation();
+
+		// Desired parent->child offset.
+		FVector TargetDelta = (ChildAnimPos - ParentAnimPos);
+
+		// Desired child position.
+		FVector ChildTarget = ParentPos + TargetDelta;
+
+		// Find vector from child to target
+		FVector Error = ChildTarget - ChildPos;
+
+		// Calculate how much to push the child towards its target
+		float Correction = FMath::Clamp<float>(ThisTimstep * TrailRelaxation, 0.f, 1.f);
+
+		// Scale correction vector and apply to get new world-space child position.
+		TrailBoneLocations[i] = ChildPos + (Error * Correction);
+
+		// If desired, prevent bones stretching too far.
+		if(bLimitStretch)
+		{
+			float RefPoseLength = TargetDelta.Size();
+			FVector CurrentDelta = TrailBoneLocations[i] - TrailBoneLocations[i-1];
+			float CurrentLength = CurrentDelta.Size();
+
+			// If we are too far - cut it back (just project towards parent particle).
+			if( (CurrentLength - RefPoseLength > StretchLimit) && CurrentLength > SMALL_NUMBER )
+			{
+				FVector CurrentDir = CurrentDelta / CurrentLength;
+				TrailBoneLocations[i] = TrailBoneLocations[i-1] + (CurrentDir * (RefPoseLength + StretchLimit));
+			}
+		}
+
+		// Modify child matrix
+		OutBoneTransforms[i] = FBoneTransform(ChildIndex, MeshBases.GetComponentSpaceTransform(ChildIndex));
+		OutBoneTransforms[i].Transform.SetTranslation(TrailBoneLocations[i]);
+
+		// Modify rotation of parent matrix to point at this one.
+
+		// Calculate the direction that parent bone is currently pointing.
+		FVector CurrentBoneDir = OutBoneTransforms[i-1].Transform.TransformVector( GetAlignVector(ChainBoneAxis, bInvertChainBoneAxis) );
+		CurrentBoneDir = CurrentBoneDir.GetSafeNormal(SMALL_NUMBER);
+
+		// Calculate vector from parent to child.
+		FVector NewBoneDir = FVector(OutBoneTransforms[i].Transform.GetTranslation() - OutBoneTransforms[i - 1].Transform.GetTranslation()).GetSafeNormal(SMALL_NUMBER);
+
+		// Calculate a quaternion that gets us from our current rotation to the desired one.
+		FQuat DeltaLookQuat = FQuat::FindBetween(CurrentBoneDir, NewBoneDir);
+		FTransform DeltaTM( DeltaLookQuat, FVector(0.f) );
+
+		// Apply to the current parent bone transform.
+		FTransform TmpMatrix = FTransform::Identity;
+		TmpMatrix.CopyRotationPart(OutBoneTransforms[i - 1].Transform);
+		TmpMatrix = TmpMatrix * DeltaTM;
+		OutBoneTransforms[i - 1].Transform.CopyRotationPart(TmpMatrix);
+	}
+
+	// For the last bone in the chain, use the rotation from the bone above it.
+	OutBoneTransforms[ChainLength - 1].Transform.CopyRotationPart(OutBoneTransforms[ChainLength - 2].Transform);
+
+	// Update OldLocalToWorld
+	OldLocalToWorld = SkelComp->GetTransformMatrix();
+}
+
+bool FAnimNode_Trail::IsValidToEvaluate(const USkeleton* Skeleton, const FBoneContainer& RequiredBones) 
+{
+	// if bones are valid
+	return (TrailBone.IsValid(RequiredBones));
+}
+
+void FAnimNode_Trail::InitializeBoneReferences(const FBoneContainer& RequiredBones) 
+{
+	TrailBone.Initialize(RequiredBones);
+}
+
+FVector FAnimNode_Trail::GetAlignVector(EAxis::Type AxisOption, bool bInvert)
+{
+	FVector AxisDir;
+
+	if (AxisOption == EAxis::X)
+	{
+		AxisDir = FVector(1, 0, 0);
+	}
+	else if (AxisOption == EAxis::Y)
+	{
+		AxisDir = FVector(0, 1, 0);
+	}
+	else
+	{
+		AxisDir = FVector(0, 0, 1);
+	}
+
+	if (bInvert)
+	{
+		AxisDir *= -1.f;
+	}
+
+	return AxisDir;
+}
diff --git a/Engine/Source/Runtime/AnimGraphRuntime/Private/BoneControllers/AnimNode_TwoBoneIK.cpp b/Engine/Source/Runtime/AnimGraphRuntime/Private/BoneControllers/AnimNode_TwoBoneIK.cpp
new file mode 100644
--- /dev/null
+++ b/Engine/Source/Runtime/AnimGraphRuntime/Private/BoneControllers/AnimNode_TwoBoneIK.cpp
@@ -0,0 +1,300 @@
+// Copyright 1998-2015 Epic Games, Inc. All Rights Reserved.
+
+#include "AnimGraphRuntimePrivatePCH.h"
+#include "BoneControllers/AnimNode_TwoBoneIK.h"
+#include "AnimationRuntime.h"
+
+/////////////////////////////////////////////////////
+// FAnimNode_TwoBoneIK
+
+FAnimNode_TwoBoneIK::FAnimNode_TwoBoneIK()
+	: EffectorLocation(FVector::ZeroVector)
+	, JointTargetLocation(FVector::ZeroVector)
+	, StretchLimits(FVector2D::ZeroVector)
+	, bTakeRotationFromEffectorSpace(false)
+	, bMaintainEffectorRelRot(false)
+	, bAllowStretching(false)
+	, EffectorLocationSpace(BCS_ComponentSpace)
+	, JointTargetLocationSpace(BCS_ComponentSpace)
+{
+}
+
+void FAnimNode_TwoBoneIK::GatherDebugData(FNodeDebugData& DebugData)
+{
+	FString DebugLine = DebugData.GetNodeName(this);
+
+	DebugLine += "(";
+	AddDebugNodeData(DebugLine);
+	DebugLine += FString::Printf(TEXT(" IKBone: %s)"), *IKBone.BoneName.ToString());
+	DebugData.AddDebugItem(DebugLine);
+
+	ComponentPose.GatherDebugData(DebugData);
+}
+
+void FAnimNode_TwoBoneIK::EvaluateBoneTransforms(USkeletalMeshComponent* SkelComp, FCSPose<FCompactPose>& MeshBases, TArray<FBoneTransform>& OutBoneTransforms)
+{
+	check(OutBoneTransforms.Num() == 0);
+
+	const FBoneContainer& BoneContainer = MeshBases.GetPose().GetBoneContainer();
+
+	// Get indices of the lower and upper limb bones and check validity.
+	bool bInvalidLimb = false;
+
+	FCompactPoseBoneIndex IKBoneCompactPoseIndex = IKBone.GetCompactPoseIndex(BoneContainer);
+
+	const FCompactPoseBoneIndex LowerLimbIndex = BoneContainer.GetParentBoneIndex(IKBoneCompactPoseIndex);
+	if (LowerLimbIndex == INDEX_NONE)
+	{
+		bInvalidLimb = true;
+	}
+
+	const FCompactPoseBoneIndex UpperLimbIndex = BoneContainer.GetParentBoneIndex(LowerLimbIndex);
+	if (UpperLimbIndex == INDEX_NONE)
+	{
+		bInvalidLimb = true;
+	}
+
+	const bool bInBoneSpace = (EffectorLocationSpace == BCS_ParentBoneSpace) || (EffectorLocationSpace == BCS_BoneSpace);
+	const int32 EffectorBoneIndex = bInBoneSpace ? BoneContainer.GetPoseBoneIndexForBoneName(EffectorSpaceBoneName) : INDEX_NONE;
+	const FCompactPoseBoneIndex EffectorSpaceBoneIndex = BoneContainer.MakeCompactPoseIndex(FMeshPoseBoneIndex(EffectorBoneIndex));
+
+	if (bInBoneSpace && (EffectorSpaceBoneIndex == INDEX_NONE))
+	{
+		bInvalidLimb = true;
+	}
+
+	// If we walked past the root, this controlled is invalid, so return no affected bones.
+	if( bInvalidLimb )
+	{
+		return;
+	}
+
+	// Get Local Space transforms for our bones. We do this first in case they already are local.
+	// As right after we get them in component space. (And that does the auto conversion).
+	// We might save one transform by doing local first...
+	const FTransform EndBoneLocalTransform = MeshBases.GetLocalSpaceTransform(IKBoneCompactPoseIndex);
+
+	// Now get those in component space...
+	FTransform LowerLimbCSTransform = MeshBases.GetComponentSpaceTransform(LowerLimbIndex);
+	FTransform UpperLimbCSTransform = MeshBases.GetComponentSpaceTransform(UpperLimbIndex);
+	FTransform EndBoneCSTransform = MeshBases.GetComponentSpaceTransform(IKBoneCompactPoseIndex);
+
+	// Get current position of root of limb.
+	// All position are in Component space.
+	const FVector RootPos = UpperLimbCSTransform.GetTranslation();
+	const FVector InitialJointPos = LowerLimbCSTransform.GetTranslation();
+	const FVector InitialEndPos = EndBoneCSTransform.GetTranslation();
+
+	// Transform EffectorLocation from EffectorLocationSpace to ComponentSpace.
+	FTransform EffectorTransform(EffectorLocation);
+	FAnimationRuntime::ConvertBoneSpaceTransformToCS(SkelComp, MeshBases, EffectorTransform, EffectorSpaceBoneIndex, EffectorLocationSpace);
+
+	// This is our reach goal.
+	FVector DesiredPos = EffectorTransform.GetTranslation();
+	FVector DesiredDelta = DesiredPos - RootPos;
+	float DesiredLength = DesiredDelta.Size();
+
+	// Check to handle case where DesiredPos is the same as RootPos.
+	FVector	DesiredDir;
+	if (DesiredLength < (float)KINDA_SMALL_NUMBER)
+	{
+		DesiredLength = (float)KINDA_SMALL_NUMBER;
+		DesiredDir = FVector(1,0,0);
+	}
+	else
+	{
+		DesiredDir = DesiredDelta / DesiredLength;
+	}
+
+	// Get joint target (used for defining plane that joint should be in).
+	FTransform JointTargetTransform(JointTargetLocation);
+	FCompactPoseBoneIndex JointTargetSpaceBoneIndex(INDEX_NONE);
+
+	if (JointTargetLocationSpace == BCS_ParentBoneSpace || JointTargetLocationSpace == BCS_BoneSpace)
+	{
+		int32 Index = BoneContainer.GetPoseBoneIndexForBoneName(JointTargetSpaceBoneName);
+		JointTargetSpaceBoneIndex = BoneContainer.MakeCompactPoseIndex(FMeshPoseBoneIndex(Index));
+	}
+	
+	FAnimationRuntime::ConvertBoneSpaceTransformToCS(SkelComp, MeshBases, JointTargetTransform, JointTargetSpaceBoneIndex, JointTargetLocationSpace);
+
+	FVector	JointTargetPos = JointTargetTransform.GetTranslation();
+	FVector JointTargetDelta = JointTargetPos - RootPos;
+	float JointTargetLength = JointTargetDelta.Size();
+
+	// Same check as above, to cover case when JointTarget position is the same as RootPos.
+	FVector JointPlaneNormal, JointBendDir;
+	if (JointTargetLength < (float)KINDA_SMALL_NUMBER)
+	{
+		JointBendDir = FVector(0,1,0);
+		JointPlaneNormal = FVector(0,0,1);
+	}
+	else
+	{
+		JointPlaneNormal = DesiredDir ^ JointTargetDelta;
+
+		// If we are trying to point the limb in the same direction that we are supposed to displace the joint in, 
+		// we have to just pick 2 random vector perp to DesiredDir and each other.
+		if (JointPlaneNormal.Size() < (float)KINDA_SMALL_NUMBER)
+		{
+			DesiredDir.FindBestAxisVectors(JointPlaneNormal, JointBendDir);
+		}
+		else
+		{
+			JointPlaneNormal.Normalize();
+
+			// Find the final member of the reference frame by removing any component of JointTargetDelta along DesiredDir.
+			// This should never leave a zero vector, because we've checked DesiredDir and JointTargetDelta are not parallel.
+			JointBendDir = JointTargetDelta - ((JointTargetDelta | DesiredDir) * DesiredDir);
+			JointBendDir.Normalize();
+		}
+	}
+
+	// Find lengths of upper and lower limb in the ref skeleton.
+	// Use actual sizes instead of ref skeleton, so we take into account translation and scaling from other bone controllers.
+	float LowerLimbLength = (InitialEndPos - InitialJointPos).Size();
+	float UpperLimbLength = (InitialJointPos - RootPos).Size();
+	float MaxLimbLength	= LowerLimbLength + UpperLimbLength;
+
+	if (bAllowStretching)
+	{
+		const float ScaleRange = StretchLimits.Y - StretchLimits.X;
+		if( ScaleRange > KINDA_SMALL_NUMBER && MaxLimbLength > KINDA_SMALL_NUMBER )
+		{
+			const float ReachRatio = DesiredLength / MaxLimbLength;
+			const float ScalingFactor = (StretchLimits.Y - 1.f) * FMath::Clamp<float>((ReachRatio - StretchLimits.X) / ScaleRange, 0.f, 1.f);
+			if (ScalingFactor > KINDA_SMALL_NUMBER)
+			{
+				LowerLimbLength *= (1.f + ScalingFactor);
+				UpperLimbLength *= (1.f + ScalingFactor);
+				MaxLimbLength	*= (1.f + ScalingFactor);
+			}
+		}
+	}
+
+	FVector OutEndPos = DesiredPos;
+	FVector OutJointPos = InitialJointPos;
+
+	// If we are trying to reach a goal beyond the length of the limb, clamp it to something solvable and extend limb fully.
+	if (DesiredLength > MaxLimbLength)
+	{
+		OutEndPos = RootPos + (MaxLimbLength * DesiredDir);
+		OutJointPos = RootPos + (UpperLimbLength * DesiredDir);
+	}
+	else
+	{
+		// So we have a triangle we know the side lengths of. We can work out the angle between DesiredDir and the direction of the upper limb
+		// using the sin rule:
+		const float TwoAB = 2.f * UpperLimbLength * DesiredLength;
+
+		const float CosAngle = (TwoAB != 0.f) ? ((UpperLimbLength*UpperLimbLength) + (DesiredLength*DesiredLength) - (LowerLimbLength*LowerLimbLength)) / TwoAB : 0.f;
+
+		// If CosAngle is less than 0, the upper arm actually points the opposite way to DesiredDir, so we handle that.
+		const bool bReverseUpperBone = (CosAngle < 0.f);
+
+		// If CosAngle is greater than 1.f, the triangle could not be made - we cannot reach the target.
+		// We just have the two limbs double back on themselves, and EndPos will not equal the desired EffectorLocation.
+		if ((CosAngle > 1.f) || (CosAngle < -1.f))
+		{
+			// Because we want the effector to be a positive distance down DesiredDir, we go back by the smaller section.
+			if (UpperLimbLength > LowerLimbLength)
+			{
+				OutJointPos = RootPos + (UpperLimbLength * DesiredDir);
+				OutEndPos = OutJointPos - (LowerLimbLength * DesiredDir);
+			}
+			else
+			{
+				OutJointPos = RootPos - (UpperLimbLength * DesiredDir);
+				OutEndPos = OutJointPos + (LowerLimbLength * DesiredDir);
+			}
+		}
+		else
+		{
+			// Angle between upper limb and DesiredDir
+			const float Angle = FMath::Acos(CosAngle);
+
+			// Now we calculate the distance of the joint from the root -> effector line.
+			// This forms a right-angle triangle, with the upper limb as the hypotenuse.
+			const float JointLineDist = UpperLimbLength * FMath::Sin(Angle);
+
+			// And the final side of that triangle - distance along DesiredDir of perpendicular.
+			// ProjJointDistSqr can't be neg, because JointLineDist must be <= UpperLimbLength because appSin(Angle) is <= 1.
+			const float ProjJointDistSqr = (UpperLimbLength*UpperLimbLength) - (JointLineDist*JointLineDist);
+			// although this shouldn't be ever negative, sometimes Xbox release produces -0.f, causing ProjJointDist to be NaN
+			// so now I branch it. 						
+			float ProjJointDist = (ProjJointDistSqr>0.f)? FMath::Sqrt(ProjJointDistSqr) : 0.f;
+			if( bReverseUpperBone )
+			{
+				ProjJointDist *= -1.f;
+			}
+
+			// So now we can work out where to put the joint!
+			OutJointPos = RootPos + (ProjJointDist * DesiredDir) + (JointLineDist * JointBendDir);
+		}
+	}
+
+	// Update transform for upper bone.
+	{
+		// Get difference in direction for old and new joint orientations
+		FVector const OldDir = (InitialJointPos - RootPos).GetSafeNormal();
+		FVector const NewDir = (OutJointPos - RootPos).GetSafeNormal();
+		// Find Delta Rotation take takes us from Old to New dir
+		FQuat const DeltaRotation = FQuat::FindBetween(OldDir, NewDir);
+		// Rotate our Joint quaternion by this delta rotation
+		UpperLimbCSTransform.SetRotation( DeltaRotation * UpperLimbCSTransform.GetRotation() );
+		// And put joint where it should be.
+		UpperLimbCSTransform.SetTranslation( RootPos );
+
+		// Order important. First bone is upper limb.
+		OutBoneTransforms.Add( FBoneTransform(UpperLimbIndex, UpperLimbCSTransform) );
+	}
+
+	// Update transform for lower bone.
+	{
+		// Get difference in direction for old and new joint orientations
+		FVector const OldDir = (InitialEndPos - InitialJointPos).GetSafeNormal();
+		FVector const NewDir = (OutEndPos - OutJointPos).GetSafeNormal();
+
+		// Find Delta Rotation take takes us from Old to New dir
+		FQuat const DeltaRotation = FQuat::FindBetween(OldDir, NewDir);
+		// Rotate our Joint quaternion by this delta rotation
+		LowerLimbCSTransform.SetRotation( DeltaRotation * LowerLimbCSTransform.GetRotation() );
+		// And put joint where it should be.
+		LowerLimbCSTransform.SetTranslation( OutJointPos );
+
+		// Order important. Second bone is lower limb.
+		OutBoneTransforms.Add( FBoneTransform(LowerLimbIndex, LowerLimbCSTransform) );
+	}
+
+	// Update transform for end bone.
+	{
+		if( bTakeRotationFromEffectorSpace )
+		{
+			EndBoneCSTransform.SetRotation( EffectorTransform.GetRotation() );
+		}
+		else if( bMaintainEffectorRelRot )
+		{
+			EndBoneCSTransform = EndBoneLocalTransform * LowerLimbCSTransform;
+		}
+
+		// Set correct location for end bone.
+		EndBoneCSTransform.SetTranslation(OutEndPos);
+
+		// Order important. Third bone is End Bone.
+		OutBoneTransforms.Add(FBoneTransform(IKBoneCompactPoseIndex, EndBoneCSTransform));
+	}
+
+	// Make sure we have correct number of bones
+	check(OutBoneTransforms.Num() == 3);
+}
+
+bool FAnimNode_TwoBoneIK::IsValidToEvaluate(const USkeleton* Skeleton, const FBoneContainer& RequiredBones) 
+{
+	// if both bones are valid
+	return (IKBone.IsValid(RequiredBones));
+}
+
+void FAnimNode_TwoBoneIK::InitializeBoneReferences(const FBoneContainer& RequiredBones) 
+{
+	IKBone.Initialize(RequiredBones);
+}
diff --git a/Engine/Source/Runtime/AnimGraphRuntime/Private/BoneControllers/AnimNode_WheelHandler.cpp b/Engine/Source/Runtime/AnimGraphRuntime/Private/BoneControllers/AnimNode_WheelHandler.cpp
new file mode 100644
--- /dev/null
+++ b/Engine/Source/Runtime/AnimGraphRuntime/Private/BoneControllers/AnimNode_WheelHandler.cpp
@@ -0,0 +1,163 @@
+// Copyright 1998-2015 Epic Games, Inc. All Rights Reserved.
+
+#include "AnimGraphRuntimePrivatePCH.h"
+#include "AnimationRuntime.h"
+#include "BoneControllers/AnimNode_WheelHandler.h"
+#include "GameFramework/WheeledVehicle.h"
+#include "Vehicles/WheeledVehicleMovementComponent.h"
+#include "Vehicles/VehicleWheel.h"
+
+/////////////////////////////////////////////////////
+// FAnimNode_WheelHandler
+
+FAnimNode_WheelHandler::FAnimNode_WheelHandler()
+{
+}
+
+void FAnimNode_WheelHandler::GatherDebugData(FNodeDebugData& DebugData)
+{
+	FString DebugLine = DebugData.GetNodeName(this);
+
+	DebugLine += "(";
+	AddDebugNodeData(DebugLine);
+	DebugLine += ")";
+
+	DebugData.AddDebugItem(DebugLine);
+	for(const auto & WheelSim : WheelSimulators)
+	{
+		if (WheelSim.BoneReference.BoneIndex != INDEX_NONE)
+		{
+			DebugLine = FString::Printf(TEXT(" [Wheel Index : %d] Bone: %s , Rotation Offset : %s, Location Offset : %s"), 
+				WheelSim.WheelIndex, *WheelSim.BoneReference.BoneName.ToString(), *WheelSim.RotOffset.ToString(), *WheelSim.LocOffset.ToString());
+		}
+		else
+		{
+			DebugLine = FString::Printf(TEXT(" [Wheel Index : %d] Bone: %s (invalid bone)"),
+					WheelSim.WheelIndex, *WheelSim.BoneReference.BoneName.ToString());
+		}
+
+		DebugData.AddDebugItem(DebugLine);
+	}
+
+	ComponentPose.GatherDebugData(DebugData);
+}
+
+void FAnimNode_WheelHandler::EvaluateBoneTransforms(USkeletalMeshComponent* SkelComp, FCSPose<FCompactPose>& MeshBases, TArray<FBoneTransform>& OutBoneTransforms)
+{
+	check(OutBoneTransforms.Num() == 0);
+
+	const FBoneContainer& BoneContainer = MeshBases.GetPose().GetBoneContainer();
+	for(const auto & WheelSim : WheelSimulators)
+	{
+		if (WheelSim.BoneReference.IsValid(BoneContainer))
+		{
+			FCompactPoseBoneIndex WheelSimBoneIndex = WheelSim.BoneReference.GetCompactPoseIndex(BoneContainer);
+
+			// the way we apply transform is same as FMatrix or FTransform
+			// we apply scale first, and rotation, and translation
+			// if you'd like to translate first, you'll need two nodes that first node does translate and second nodes to rotate. 
+			FTransform NewBoneTM = MeshBases.GetComponentSpaceTransform(WheelSimBoneIndex);
+
+			FAnimationRuntime::ConvertCSTransformToBoneSpace(SkelComp, MeshBases, NewBoneTM, WheelSimBoneIndex, BCS_ComponentSpace);
+			
+			// Apply rotation offset
+			const FQuat BoneQuat(WheelSim.RotOffset);
+			NewBoneTM.SetRotation(BoneQuat * NewBoneTM.GetRotation());
+
+			// Apply loc offset
+			NewBoneTM.AddToTranslation(WheelSim.LocOffset);
+
+			// Convert back to Component Space.
+			FAnimationRuntime::ConvertBoneSpaceTransformToCS(SkelComp, MeshBases, NewBoneTM, WheelSimBoneIndex, BCS_ComponentSpace);
+
+			// add back to it
+			OutBoneTransforms.Add(FBoneTransform(WheelSimBoneIndex, NewBoneTM));
+		}
+	}
+}
+
+bool FAnimNode_WheelHandler::IsValidToEvaluate(const USkeleton* Skeleton, const FBoneContainer& RequiredBones) 
+{
+	// if both bones are valid
+	for(const auto & WheelSim : WheelSimulators)
+	{
+		// if one of them is valid
+		if (WheelSim.BoneReference.IsValid(RequiredBones) == true)
+		{
+			return true;
+		}
+	}
+
+	return false;
+}
+
+void FAnimNode_WheelHandler::InitializeBoneReferences(const FBoneContainer& RequiredBones) 
+{
+	for (auto & WheelSim : WheelSimulators)
+	{
+		WheelSim.BoneReference.Initialize(RequiredBones);
+	}
+
+	// sort by bone indices
+	WheelSimulators.Sort([](FWheelSimulator L, FWheelSimulator R) { return L.BoneReference.BoneIndex < R.BoneReference.BoneIndex; });
+}
+
+void FAnimNode_WheelHandler::Update(const FAnimationUpdateContext& Context)
+{
+	if(VehicleSimComponent)
+	{
+		for(auto & WheelSim : WheelSimulators)
+		{
+			if(VehicleSimComponent->Wheels.IsValidIndex(WheelSim.WheelIndex))
+			{
+				UVehicleWheel* Wheel = VehicleSimComponent->Wheels[WheelSim.WheelIndex];
+				if(Wheel != nullptr)
+				{
+					WheelSim.RotOffset.Pitch	= Wheel->GetRotationAngle();
+					WheelSim.RotOffset.Yaw		= Wheel->GetSteerAngle();
+					WheelSim.RotOffset.Roll		= 0.f;
+
+					WheelSim.LocOffset.X		= 0.f;
+					WheelSim.LocOffset.Y		= 0.f;
+					WheelSim.LocOffset.Z		= Wheel->GetSuspensionOffset();
+				}
+			}
+		}
+	}
+
+	FAnimNode_SkeletalControlBase::Update(Context);
+}
+
+void FAnimNode_WheelHandler::Initialize(const FAnimationInitializeContext& Context)
+{
+	// TODO: only check vehicle anim instance
+	// UVehicleAnimInstance
+	AWheeledVehicle * Vehicle = Cast<AWheeledVehicle> ((CastChecked<USkeletalMeshComponent> (Context.AnimInstance->GetOuter()))->GetOwner());
+	;
+	// we only support vehicle for this class
+	if(Vehicle != nullptr)
+	{
+		VehicleSimComponent = Vehicle->GetVehicleMovementComponent();
+
+		int32 NumOfwheels = VehicleSimComponent->WheelSetups.Num();
+		if(NumOfwheels > 0)
+		{
+			WheelSimulators.Empty(NumOfwheels);
+			WheelSimulators.AddZeroed(NumOfwheels);
+			// now add wheel data
+			for(int32 WheelIndex = 0; WheelIndex<WheelSimulators.Num(); ++WheelIndex)
+			{
+				FWheelSimulator & WheelSim = WheelSimulators[WheelIndex];
+				const FWheelSetup& WheelSetup = VehicleSimComponent->WheelSetups[WheelIndex];
+
+				// set data
+				WheelSim.WheelIndex = WheelIndex;
+				WheelSim.BoneReference.BoneName = WheelSetup.BoneName;
+				WheelSim.LocOffset = FVector::ZeroVector;
+				WheelSim.RotOffset = FRotator::ZeroRotator;
+			}
+		}
+	}
+
+	FAnimNode_SkeletalControlBase::Initialize(Context);
+}
\ No newline at end of file
diff --git a/Engine/Source/Runtime/AnimGraphRuntime/Public/BoneControllers/AnimNode_BoneDrivenController.h b/Engine/Source/Runtime/AnimGraphRuntime/Public/BoneControllers/AnimNode_BoneDrivenController.h
new file mode 100644
--- /dev/null
+++ b/Engine/Source/Runtime/AnimGraphRuntime/Public/BoneControllers/AnimNode_BoneDrivenController.h
@@ -0,0 +1,79 @@
+// Copyright 1998-2015 Epic Games, Inc. All Rights Reserved.
+
+#pragma once
+
+#include "AnimNode_SkeletalControlBase.h"
+#include "AnimNode_BoneDrivenController.generated.h"
+
+// Evaluation of the bone transforms relies on the size and ordering of this
+// enum, if this needs to change make sure EvaluateBoneTransforms is updated.
+UENUM()
+namespace EComponentType
+{
+	enum Type
+	{
+		None = 0,
+		TranslationX,
+		TranslationY,
+		TranslationZ,
+		RotationX,
+		RotationY,
+		RotationZ,
+		Scale
+	};
+}
+
+USTRUCT()
+struct ANIMGRAPHRUNTIME_API FAnimNode_BoneDrivenController : public FAnimNode_SkeletalControlBase
+{
+	GENERATED_USTRUCT_BODY()
+
+	// Bone to use as controller input
+	UPROPERTY(EditAnywhere, BlueprintReadWrite, Category = Bones)
+	FBoneReference SourceBone;
+
+	// Bone to drive using controller input
+	UPROPERTY(EditAnywhere, BlueprintReadWrite, Category = Bones)
+	FBoneReference TargetBone;
+
+	// Transform component to use as input
+	UPROPERTY(EditAnywhere, BlueprintReadWrite, Category = Controller)
+	TEnumAsByte<EComponentType::Type> SourceComponent;
+
+	// Transform component to write to
+	UPROPERTY(EditAnywhere, BlueprintReadWrite, Category = Controller)
+	TEnumAsByte<EComponentType::Type> TargetComponent;
+
+	// Multiplier to apply to the input value
+	UPROPERTY(EditAnywhere, BlueprintReadWrite, Category=Controller)
+	float Multiplier;
+
+	// Whether or not to use the range limits
+	UPROPERTY(EditAnywhere, BlueprintReadWrite, Category=Range)
+	bool bUseRange;
+
+	// Min limit of the output value
+	UPROPERTY(EditAnywhere, BlueprintReadWrite, Category=Range)
+	float RangeMin;
+
+	// Max limit of the output value
+	UPROPERTY(EditAnywhere, BlueprintReadWrite, Category=Range)
+	float RangeMax;
+
+	FAnimNode_BoneDrivenController();
+
+	// FAnimNode_Base interface
+	virtual void GatherDebugData(FNodeDebugData& DebugData) override;
+	//////////////////////////////////////////////////////////////////////////
+
+	// FAnimNode_SkeletalControlBase interface
+	virtual void EvaluateBoneTransforms(USkeletalMeshComponent* SkelComp, FCSPose<FCompactPose>& MeshBases, TArray<FBoneTransform>& OutBoneTransforms) override;
+	virtual bool IsValidToEvaluate(const USkeleton* Skeleton, const FBoneContainer& RequiredBones) override;
+	//////////////////////////////////////////////////////////////////////////
+
+protected:
+	
+	// FAnimNode_SkeletalControlBase interface
+	virtual void InitializeBoneReferences(const FBoneContainer& RequiredBones) override;
+	//////////////////////////////////////////////////////////////////////////
+};
\ No newline at end of file
diff --git a/Engine/Source/Runtime/AnimGraphRuntime/Public/BoneControllers/AnimNode_CopyBone.h b/Engine/Source/Runtime/AnimGraphRuntime/Public/BoneControllers/AnimNode_CopyBone.h
new file mode 100644
--- /dev/null
+++ b/Engine/Source/Runtime/AnimGraphRuntime/Public/BoneControllers/AnimNode_CopyBone.h
@@ -0,0 +1,52 @@
+// Copyright 1998-2015 Epic Games, Inc. All Rights Reserved.
+
+#pragma once
+
+#include "AnimNode_SkeletalControlBase.h"
+#include "AnimNode_CopyBone.generated.h"
+
+/**
+ *	Simple controller to copy a bone's transform to another one.
+ */
+
+USTRUCT()
+struct ANIMGRAPHRUNTIME_API FAnimNode_CopyBone : public FAnimNode_SkeletalControlBase
+{
+	GENERATED_USTRUCT_BODY()
+
+	/** Name of bone to control. This is the main bone chain to modify from. **/
+	UPROPERTY(EditAnywhere, BlueprintReadWrite, Category=Copy) 
+	FBoneReference TargetBone;
+
+	/** Source Bone Name to get transform from */
+	UPROPERTY(EditAnywhere, BlueprintReadWrite, Category=Copy)
+	FBoneReference SourceBone;
+
+	/** If Translation should be copied */
+	UPROPERTY(EditAnywhere, BlueprintReadWrite, Category=Copy, meta=(PinShownByDefault))
+	bool bCopyTranslation;
+
+	/** If Rotation should be copied */
+	UPROPERTY(EditAnywhere, BlueprintReadWrite, Category=Copy, meta=(PinShownByDefault))
+	bool bCopyRotation;
+
+	/** If Scale should be copied */
+	UPROPERTY(EditAnywhere, BlueprintReadWrite, Category=Copy, meta=(PinShownByDefault))
+	bool bCopyScale;
+
+	FAnimNode_CopyBone();
+
+	// FAnimNode_Base interface
+	virtual void GatherDebugData(FNodeDebugData& DebugData) override;
+	// End of FAnimNode_Base interface
+
+	// FAnimNode_SkeletalControlBase interface
+	virtual void EvaluateBoneTransforms(USkeletalMeshComponent* SkelComp, FCSPose<FCompactPose>& MeshBases, TArray<FBoneTransform>& OutBoneTransforms) override;
+	virtual bool IsValidToEvaluate(const USkeleton* Skeleton, const FBoneContainer& RequiredBones) override;
+	// End of FAnimNode_SkeletalControlBase interface
+
+private:
+	// FAnimNode_SkeletalControlBase interface
+	virtual void InitializeBoneReferences(const FBoneContainer& RequiredBones) override;
+	// End of FAnimNode_SkeletalControlBase interface
+};
diff --git a/Engine/Source/Runtime/AnimGraphRuntime/Public/BoneControllers/AnimNode_Fabrik.h b/Engine/Source/Runtime/AnimGraphRuntime/Public/BoneControllers/AnimNode_Fabrik.h
new file mode 100644
--- /dev/null
+++ b/Engine/Source/Runtime/AnimGraphRuntime/Public/BoneControllers/AnimNode_Fabrik.h
@@ -0,0 +1,105 @@
+// Copyright 1998-2015 Epic Games, Inc. All Rights Reserved.
+
+#pragma once
+
+#include "AnimNode_SkeletalControlBase.h"
+#include "AnimNode_Fabrik.generated.h"
+
+/**
+*	Controller which implements the FABRIK IK approximation algorithm -  see http://andreasaristidou.com/publications/FABRIK.pdf for details
+*/
+
+/** Transient structure for FABRIK node evaluation */
+struct FABRIKChainLink
+{
+public:
+	/** Position of bone in component space. */
+	FVector Position;
+
+	/** Distance to its parent link. */
+	float Length;
+
+	/** Bone Index in SkeletalMesh */
+	FCompactPoseBoneIndex BoneIndex;
+
+	/** Transform Index that this control will output */
+	int32 TransformIndex;
+
+	/** Child bones which are overlapping this bone. 
+	 * They have a zero length distance, so they will inherit this bone's transformation. */
+	TArray<int32> ChildZeroLengthTransformIndices;
+
+	FABRIKChainLink()
+		: Position(FVector::ZeroVector)
+		, Length(0.f)
+		, BoneIndex(INDEX_NONE)
+		, TransformIndex(INDEX_NONE)
+	{
+	}
+
+	FABRIKChainLink(const FVector& InPosition, const float& InLength, const FCompactPoseBoneIndex& InBoneIndex, const int32& InTransformIndex)
+		: Position(InPosition)
+		, Length(InLength)
+		, BoneIndex(InBoneIndex)
+		, TransformIndex(InTransformIndex)
+	{
+	}
+};
+
+USTRUCT()
+struct ANIMGRAPHRUNTIME_API FAnimNode_Fabrik : public FAnimNode_SkeletalControlBase
+{
+	GENERATED_USTRUCT_BODY()
+
+	/** Coordinates for target location of tip bone - if EffectorLocationSpace is bone, this is the offset from Target Bone to use as target location*/
+	UPROPERTY(EditAnywhere, BlueprintReadWrite, Category = EndEffector, meta = (PinShownByDefault))
+	FTransform EffectorTransform;
+
+	/** Reference frame of Effector Transform. */
+	UPROPERTY(EditAnywhere, BlueprintReadWrite, Category = EndEffector)
+	TEnumAsByte<enum EBoneControlSpace> EffectorTransformSpace;
+
+	/** If EffectorTransformSpace is a bone, this is the bone to use. **/
+	UPROPERTY(EditAnywhere, BlueprintReadWrite, Category = EndEffector)
+	FBoneReference EffectorTransformBone;
+
+	UPROPERTY(EditAnywhere, BlueprintReadWrite, Category = EndEffector)
+	TEnumAsByte<enum EBoneRotationSource> EffectorRotationSource;
+
+	/** Name of tip bone */
+	UPROPERTY(EditAnywhere, BlueprintReadWrite, Category = Solver)
+	FBoneReference TipBone;
+
+	/** Name of the root bone*/
+	UPROPERTY(EditAnywhere, BlueprintReadWrite, Category = Solver)
+	FBoneReference RootBone;
+
+	/** Tolerance for final tip location delta from EffectorLocation*/
+	UPROPERTY(EditAnywhere, BlueprintReadWrite, Category = Solver)
+	float Precision;
+
+	/** Maximum number of iterations allowed, to control performance. */
+	UPROPERTY(EditAnywhere, BlueprintReadWrite, Category = Solver)
+	int32 MaxIterations;
+
+	/** Toggle drawing of axes to debug joint rotation*/
+	UPROPERTY(EditAnywhere, BlueprintReadWrite, Category = Solver)
+	bool bEnableDebugDraw;
+
+public:
+	FAnimNode_Fabrik();
+
+	// FAnimNode_SkeletalControlBase interface
+	virtual void EvaluateBoneTransforms(USkeletalMeshComponent* SkelComp, FCSPose<FCompactPose>& MeshBases, TArray<FBoneTransform>& OutBoneTransforms) override;
+	virtual bool IsValidToEvaluate(const USkeleton* Skeleton, const FBoneContainer& RequiredBones) override;
+	// End of FAnimNode_SkeletalControlBase interface
+
+
+private:
+	// FAnimNode_SkeletalControlBase interface
+	virtual void InitializeBoneReferences(const FBoneContainer& RequiredBones) override;
+	// End of FAnimNode_SkeletalControlBase interface
+
+	// Convenience function to get current (pre-translation iteration) component space location of bone by bone index
+	FVector GetCurrentLocation(FCSPose<FCompactPose>& MeshBases, const FCompactPoseBoneIndex& BoneIndex);
+};
\ No newline at end of file
diff --git a/Engine/Source/Runtime/AnimGraphRuntime/Public/BoneControllers/AnimNode_HandIKRetargeting.h b/Engine/Source/Runtime/AnimGraphRuntime/Public/BoneControllers/AnimNode_HandIKRetargeting.h
new file mode 100644
--- /dev/null
+++ b/Engine/Source/Runtime/AnimGraphRuntime/Public/BoneControllers/AnimNode_HandIKRetargeting.h
@@ -0,0 +1,59 @@
+// Copyright 1998-2015 Epic Games, Inc. All Rights Reserved.
+
+#pragma once
+
+#include "AnimNode_SkeletalControlBase.h"
+#include "AnimNode_HandIKRetargeting.generated.h"
+
+/**
+ * Node to handle re-targeting of Hand IK bone chain.
+ * It looks at position in Mesh Space of Left and Right IK bones, and moves Left and Right IK bones to those.
+ * based on HandFKWeight. (0 = favor left hand, 1 = favor right hand, 0.5 = equal weight).
+ * This is used so characters of different proportions can handle the same props.
+ */
+
+USTRUCT()
+struct ANIMGRAPHRUNTIME_API FAnimNode_HandIKRetargeting : public FAnimNode_SkeletalControlBase
+{
+	GENERATED_USTRUCT_BODY()
+
+	/** Bone for Right Hand FK */
+	UPROPERTY(EditAnywhere, BlueprintReadWrite, Category = "HandIKRetargeting")
+	FBoneReference RightHandFK;
+
+	/** Bone for Left Hand FK */
+	UPROPERTY(EditAnywhere, BlueprintReadWrite, Category = "HandIKRetargeting")
+	FBoneReference LeftHandFK;
+
+	/** Bone for Right Hand IK */
+	UPROPERTY(EditAnywhere, BlueprintReadWrite, Category = "HandIKRetargeting")
+	FBoneReference RightHandIK;
+
+	/** Bone for Left Hand FK */
+	UPROPERTY(EditAnywhere, BlueprintReadWrite, Category = "HandIKRetargeting")
+	FBoneReference LeftHandIK;
+
+	/** IK Bones to move. */
+	UPROPERTY(EditAnywhere, BlueprintReadWrite, Category = "HandIKRetargeting")
+	TArray<FBoneReference> IKBonesToMove;
+
+	/** Which hand to favor. 0.5 is equal weight for both, 1 = right hand, 0 = left hand. */
+	UPROPERTY(EditAnywhere, BlueprintReadWrite, Category = "HandIKRetargeting", meta = (PinShownByDefault))
+	float HandFKWeight;
+
+	FAnimNode_HandIKRetargeting();
+
+	// FAnimNode_Base interface
+	virtual void GatherDebugData(FNodeDebugData& DebugData) override;
+	// End of FAnimNode_Base interface
+
+	// FAnimNode_SkeletalControlBase interface
+	virtual void EvaluateBoneTransforms(USkeletalMeshComponent* SkelComp, FCSPose<FCompactPose>& MeshBases, TArray<FBoneTransform>& OutBoneTransforms) override;
+	virtual bool IsValidToEvaluate(const USkeleton* Skeleton, const FBoneContainer& RequiredBones) override;
+	// End of FAnimNode_SkeletalControlBase interface
+
+private:
+	// FAnimNode_SkeletalControlBase interface
+	virtual void InitializeBoneReferences(const FBoneContainer& RequiredBones) override;
+	// End of FAnimNode_SkeletalControlBase interface
+};
diff --git a/Engine/Source/Runtime/AnimGraphRuntime/Public/BoneControllers/AnimNode_LookAt.h b/Engine/Source/Runtime/AnimGraphRuntime/Public/BoneControllers/AnimNode_LookAt.h
new file mode 100644
--- /dev/null
+++ b/Engine/Source/Runtime/AnimGraphRuntime/Public/BoneControllers/AnimNode_LookAt.h
@@ -0,0 +1,130 @@
+// Copyright 1998-2015 Epic Games, Inc. All Rights Reserved.
+
+#pragma once
+
+#include "AnimNode_SkeletalControlBase.h"
+#include "AnimNode_LookAt.generated.h"
+
+UENUM()
+namespace EAxisOption
+{
+	enum Type
+	{
+		X,
+		Y,
+		Z, 
+		X_Neg, 
+		Y_Neg, 
+		Z_Neg
+	};
+}
+
+
+UENUM()
+/** Various ways to interpolate TAlphaBlend. */
+namespace EInterpolationBlend
+{
+	enum Type
+	{
+		Linear,
+		Cubic,
+		Sinusoidal,
+		EaseInOutExponent2,
+		EaseInOutExponent3,
+		EaseInOutExponent4,
+		EaseInOutExponent5,
+		MAX
+	};
+}
+
+/**
+ *	Simple controller that make a bone to look at the point or another bone
+ */
+USTRUCT()
+struct ANIMGRAPHRUNTIME_API FAnimNode_LookAt : public FAnimNode_SkeletalControlBase
+{
+	GENERATED_USTRUCT_BODY()
+
+	/** Name of bone to control. This is the main bone chain to modify from. **/
+	UPROPERTY(EditAnywhere, BlueprintReadWrite, Category=SkeletalControl) 
+	FBoneReference BoneToModify;
+
+	/** Target Bone to look at - you can't use LookAtLocation as alternative as you'll get a delay on bone location if you query directly **/
+	UPROPERTY(EditAnywhere, BlueprintReadWrite, Category=SkeletalControl)
+	FBoneReference LookAtBone;
+
+	/** Target Location in world space if LookAtBone is empty */
+	UPROPERTY(EditAnywhere, BlueprintReadWrite, Category=SkeletalControl, meta=(PinHiddenByDefault))
+	FVector LookAtLocation;
+
+	/** Look at axis, which axis to align to look at point */
+	UPROPERTY(EditAnywhere, BlueprintReadWrite, Category=SkeletalControl) 
+	TEnumAsByte<EAxisOption::Type>	LookAtAxis;
+
+	/** Look at Clamp value in degree - if you're look at axis is Z, only X, Y degree of clamp will be used*/
+	UPROPERTY(EditAnywhere, BlueprintReadWrite, Category=SkeletalControl, meta=(PinHiddenByDefault))
+	float LookAtClamp;
+
+	UPROPERTY(EditAnywhere, BlueprintReadWrite, Category=SkeletalControl, meta=(PinHiddenByDefault))
+	TEnumAsByte<EInterpolationBlend::Type>	InterpolationType;
+
+	UPROPERTY(EditAnywhere, BlueprintReadWrite, Category=SkeletalControl, meta=(PinHiddenByDefault))
+	float	InterpolationTime;
+
+	UPROPERTY(EditAnywhere, BlueprintReadWrite, Category=SkeletalControl, meta=(PinHiddenByDefault))
+	float	InterpolationTriggerThreashold;
+
+	UPROPERTY(EditAnywhere, Category=SkeletalControl)
+	bool	bEnableDebug;
+
+	/** Debug transient data */
+	FVector CurrentLookAtLocation;
+
+	/** Current Target Location */
+	FVector CurrentTargetLocation;
+	FVector PreviousTargetLocation;
+
+	/** Current Alpha */
+	float AccumulatedInterpoolationTime;
+
+	// in the future, it would be nice to have more options, -i.e. lag, interpolation speed
+	FAnimNode_LookAt();
+
+	// FAnimNode_Base interface
+	virtual void GatherDebugData(FNodeDebugData& DebugData) override;
+	virtual void Update(const FAnimationUpdateContext& Context) override;
+	// End of FAnimNode_Base interface
+
+	// FAnimNode_SkeletalControlBase interface
+	virtual void EvaluateBoneTransforms(USkeletalMeshComponent* SkelComp, FCSPose<FCompactPose>& MeshBases, TArray<FBoneTransform>& OutBoneTransforms) override;
+	virtual bool IsValidToEvaluate(const USkeleton* Skeleton, const FBoneContainer& RequiredBones) override;
+	// End of FAnimNode_SkeletalControlBase interface
+
+private:
+	// FAnimNode_SkeletalControlBase interface
+	virtual void InitializeBoneReferences(const FBoneContainer& RequiredBones) override;
+	// End of FAnimNode_SkeletalControlBase interface
+
+	FVector GetAlignVector(const FTransform& Transform, EAxisOption::Type AxisOption);
+
+	EAlphaBlendType GetInterpolationType()
+	{
+		switch (InterpolationType)
+		{
+		case EInterpolationBlend::Cubic:
+			return ABT_Cubic;
+		case EInterpolationBlend::Sinusoidal:
+			return ABT_Cubic;
+		case EInterpolationBlend::EaseInOutExponent2:
+			return ABT_EaseInOutExponent2;
+		case EInterpolationBlend::EaseInOutExponent3:
+			return ABT_EaseInOutExponent3;
+		case EInterpolationBlend::EaseInOutExponent4:
+			return ABT_EaseInOutExponent4;
+		case EInterpolationBlend::EaseInOutExponent5:
+			return ABT_EaseInOutExponent5;
+		}
+
+		return ABT_Linear;
+	}
+};
diff --git a/Engine/Source/Runtime/AnimGraphRuntime/Public/BoneControllers/AnimNode_ModifyBone.h b/Engine/Source/Runtime/AnimGraphRuntime/Public/BoneControllers/AnimNode_ModifyBone.h
new file mode 100644
--- /dev/null
+++ b/Engine/Source/Runtime/AnimGraphRuntime/Public/BoneControllers/AnimNode_ModifyBone.h
@@ -0,0 +1,84 @@
+// Copyright 1998-2015 Epic Games, Inc. All Rights Reserved.
+
+#pragma once
+
+#include "AnimNode_SkeletalControlBase.h"
+#include "AnimNode_ModifyBone.generated.h"
+
+UENUM()
+enum EBoneModificationMode
+{
+	// The modifier ignores this channel (keeps the existing bone translation, rotation, or scale)
+	BMM_Ignore,
+
+	// The modifier replaces the existing translation, rotation, or scale
+	BMM_Replace,
+
+	// The modifier adds to the existing translation, rotation, or scale
+	BMM_Additive
+};
+
+/**
+ *	Simple controller that replaces or adds to the translation/rotation of a single bone.
+ */
+USTRUCT()
+struct ANIMGRAPHRUNTIME_API FAnimNode_ModifyBone : public FAnimNode_SkeletalControlBase
+{
+	GENERATED_USTRUCT_BODY()
+
+	/** Name of bone to control. This is the main bone chain to modify from. **/
+	UPROPERTY(EditAnywhere, BlueprintReadWrite, Category=SkeletalControl) 
+	FBoneReference BoneToModify;
+
+	/** New translation of bone to apply. */
+	UPROPERTY(EditAnywhere, BlueprintReadWrite, Category=Translation, meta=(PinShownByDefault))
+	FVector Translation;
+
+	/** New rotation of bone to apply. */
+	UPROPERTY(EditAnywhere, BlueprintReadWrite, Category=Rotation, meta=(PinShownByDefault))
+	FRotator Rotation;
+
+	/** New Scale of bone to apply. This is only worldspace. */
+	UPROPERTY(EditAnywhere, BlueprintReadWrite, Category=Scale, meta=(PinShownByDefault))
+	FVector Scale;
+
+	/** Whether and how to modify the translation of this bone. */
+	UPROPERTY(EditAnywhere, BlueprintReadWrite, Category=Translation)
+	TEnumAsByte<EBoneModificationMode> TranslationMode;
+
+	/** Whether and how to modify the translation of this bone. */
+	UPROPERTY(EditAnywhere, BlueprintReadWrite, Category=Rotation)
+	TEnumAsByte<EBoneModificationMode> RotationMode;
+
+	/** Whether and how to modify the translation of this bone. */
+	UPROPERTY(EditAnywhere, BlueprintReadWrite, Category=Scale)
+	TEnumAsByte<EBoneModificationMode> ScaleMode;
+
+	/** Reference frame to apply Translation in. */
+	UPROPERTY(EditAnywhere, BlueprintReadWrite, Category=Translation)
+	TEnumAsByte<enum EBoneControlSpace> TranslationSpace;
+
+	/** Reference frame to apply Rotation in. */
+	UPROPERTY(EditAnywhere, BlueprintReadWrite, Category=Rotation)
+	TEnumAsByte<enum EBoneControlSpace> RotationSpace;
+
+	/** Reference frame to apply Scale in. */
+	UPROPERTY(EditAnywhere, BlueprintReadWrite, Category=Scale)
+	TEnumAsByte<enum EBoneControlSpace> ScaleSpace;
+
+	FAnimNode_ModifyBone();
+
+	// FAnimNode_Base interface
+	virtual void GatherDebugData(FNodeDebugData& DebugData) override;
+	// End of FAnimNode_Base interface
+
+	// FAnimNode_SkeletalControlBase interface
+	virtual void EvaluateBoneTransforms(USkeletalMeshComponent* SkelComp, FCSPose<FCompactPose>& MeshBases, TArray<FBoneTransform>& OutBoneTransforms) override;
+	virtual bool IsValidToEvaluate(const USkeleton* Skeleton, const FBoneContainer& RequiredBones) override;
+	// End of FAnimNode_SkeletalControlBase interface
+
+private:
+	// FAnimNode_SkeletalControlBase interface
+	virtual void InitializeBoneReferences(const FBoneContainer& RequiredBones) override;
+	// End of FAnimNode_SkeletalControlBase interface
+};
diff --git a/Engine/Source/Runtime/AnimGraphRuntime/Public/BoneControllers/AnimNode_RotationMultiplier.h b/Engine/Source/Runtime/AnimGraphRuntime/Public/BoneControllers/AnimNode_RotationMultiplier.h
new file mode 100644
--- /dev/null
+++ b/Engine/Source/Runtime/AnimGraphRuntime/Public/BoneControllers/AnimNode_RotationMultiplier.h
@@ -0,0 +1,64 @@
+// Copyright 1998-2015 Epic Games, Inc. All Rights Reserved.
+
+#pragma once
+
+#include "AnimNode_SkeletalControlBase.h"
+#include "AnimNode_RotationMultiplier.generated.h"
+
+UENUM()
+enum EBoneAxis
+{
+	BA_X,
+	BA_Y,
+	BA_Z
+};
+
+/**
+ *	Simple controller that multiplies scalar value to the translation/rotation/scale of a single bone.
+ */
+
+USTRUCT()
+struct ANIMGRAPHRUNTIME_API FAnimNode_RotationMultiplier : public FAnimNode_SkeletalControlBase
+{
+	GENERATED_USTRUCT_BODY()
+
+	/** Name of bone to control. This is the main bone chain to modify from. **/
+	UPROPERTY(EditAnywhere, BlueprintReadWrite, Category=Multiplier) 
+	FBoneReference	TargetBone;
+
+	/** Source to get transform from **/
+	UPROPERTY(EditAnywhere, BlueprintReadWrite, Category=Multiplier)
+	FBoneReference	SourceBone;
+
+	// To make these to be easily pin-hookable, I'm not making it struct, but each variable
+	// 0.f is invalid, and default
+	UPROPERTY(EditAnywhere, BlueprintReadWrite, Category=Multiplier, meta=(PinShownByDefault))
+	float	Multiplier;
+
+	UPROPERTY(EditAnywhere, BlueprintReadWrite, Category=Multiplier)
+	TEnumAsByte<EBoneAxis> RotationAxisToRefer;
+	
+	UPROPERTY(EditAnywhere, BlueprintReadWrite, Category = Multiplier)
+	bool bIsAdditive;	
+
+	FAnimNode_RotationMultiplier();
+
+	// FAnimNode_Base interface
+	virtual void GatherDebugData(FNodeDebugData& DebugData) override;
+	// End of FAnimNode_Base interface
+
+	// FAnimNode_SkeletalControlBase interface
+	virtual void EvaluateBoneTransforms(USkeletalMeshComponent* SkelComp, FCSPose<FCompactPose>& MeshBases, TArray<FBoneTransform>& OutBoneTransforms) override;
+	virtual bool IsValidToEvaluate(const USkeleton* Skeleton, const FBoneContainer& RequiredBones) override;
+	// End of FAnimNode_SkeletalControlBase interface
+
+private:
+	// Extract Delta Quat of rotation around Axis of animation and reference pose for the SourceBoneIndex
+	FQuat ExtractAngle(const FTransform& RefPoseTransform, const FTransform& LocalBoneTransform, const EBoneAxis Axis);
+	// Multiply scalar value Multiplier to the delta Quat of SourceBone Index's rotation
+	FQuat MultiplyQuatBasedOnSourceIndex(const FTransform& RefPoseTransform, const FTransform& LocalBoneTransform, const EBoneAxis Axis, float InMultiplier, const FQuat& ReferenceQuat);
+
+	// FAnimNode_SkeletalControlBase interface
+	virtual void InitializeBoneReferences(const FBoneContainer& RequiredBones) override;
+	// End of FAnimNode_SkeletalControlBase interface
+};
diff --git a/Engine/Source/Runtime/AnimGraphRuntime/Public/BoneControllers/AnimNode_SkeletalControlBase.h b/Engine/Source/Runtime/AnimGraphRuntime/Public/BoneControllers/AnimNode_SkeletalControlBase.h
new file mode 100644
--- /dev/null
+++ b/Engine/Source/Runtime/AnimGraphRuntime/Public/BoneControllers/AnimNode_SkeletalControlBase.h
@@ -0,0 +1,63 @@
+// Copyright 1998-2015 Epic Games, Inc. All Rights Reserved.
+
+/**
+ *	Abstract base class for a skeletal controller.
+ *	A SkelControl is a module that can modify the position or orientation of a set of bones in a skeletal mesh in some programmatic way.
+ */
+
+#pragma once
+
+#include "Animation/AnimNodeBase.h"
+#include "Animation/InputScaleBias.h"
+#include "AnimNode_SkeletalControlBase.generated.h"
+
+
+USTRUCT()
+struct ANIMGRAPHRUNTIME_API FAnimNode_SkeletalControlBase : public FAnimNode_Base
+{
+	GENERATED_USTRUCT_BODY()
+
+	// Input link
+	UPROPERTY(EditAnywhere, BlueprintReadWrite, Category=Links)
+	FComponentSpacePoseLink ComponentPose;
+
+	// Current strength of the skeletal control
+	UPROPERTY(EditAnywhere, BlueprintReadWrite, Category=Settings, meta=(PinShownByDefault))
+	mutable float Alpha;
+
+	UPROPERTY(EditAnywhere, BlueprintReadWrite, Category=Settings)
+	FInputScaleBias AlphaScaleBias;
+
+public:
+
+	FAnimNode_SkeletalControlBase()
+		: Alpha(1.0f)
+	{
+	}
+
+public:
+#if WITH_EDITORONLY_DATA
+	// forwarded pose data from the wired node which current node's skeletal control is not applied yet
+	FCSPose<FCompactPose> ForwardedPose;
+#endif //#if WITH_EDITORONLY_DATA
+
+	// FAnimNode_Base interface
+	virtual void Initialize(const FAnimationInitializeContext& Context) override;
+	virtual void CacheBones(const FAnimationCacheBonesContext& Context)  override;
+	virtual void Update(const FAnimationUpdateContext& Context) override;
+	virtual void EvaluateComponentSpace(FComponentSpacePoseContext& Output) override;
+	// End of FAnimNode_Base interface
+	
+protected:
+	// Interface for derived skeletal controls to implement
+
+	// Evaluate the new component-space transforms for the affected bones.
+	virtual void EvaluateBoneTransforms(USkeletalMeshComponent* SkelComp, FCSPose<FCompactPose>& MeshBases, TArray<FBoneTransform>& OutBoneTransforms) {}
+	// return true if it is valid to Evaluate
+	virtual bool IsValidToEvaluate(const USkeleton* Skeleton, const FBoneContainer& RequiredBones) { return false; }
+	// initialize any bone references you have
+	virtual void InitializeBoneReferences(const FBoneContainer& RequiredBones){};
+
+	/** Allow base to add info to the node debug output */
+	void AddDebugNodeData(FString& OutDebugData);
+};
diff --git a/Engine/Source/Runtime/AnimGraphRuntime/Public/BoneControllers/AnimNode_SpringBone.h b/Engine/Source/Runtime/AnimGraphRuntime/Public/BoneControllers/AnimNode_SpringBone.h
new file mode 100644
--- /dev/null
+++ b/Engine/Source/Runtime/AnimGraphRuntime/Public/BoneControllers/AnimNode_SpringBone.h
@@ -0,0 +1,103 @@
+// Copyright 1998-2015 Epic Games, Inc. All Rights Reserved.
+
+#pragma once
+
+#include "AnimNode_SkeletalControlBase.h"
+#include "AnimNode_SpringBone.generated.h"
+
+/**
+ *	Simple controller that replaces or adds to the translation/rotation of a single bone.
+ */
+
+USTRUCT()
+struct ANIMGRAPHRUNTIME_API FAnimNode_SpringBone : public FAnimNode_SkeletalControlBase
+{
+	GENERATED_USTRUCT_BODY()
+
+	/** Name of bone to control. This is the main bone chain to modify from. **/
+	UPROPERTY(EditAnywhere, BlueprintReadWrite, Category=Spring) 
+	FBoneReference SpringBone;
+
+	/** Limit the amount that a bone can stretch from its ref-pose length. */
+	UPROPERTY(EditAnywhere, BlueprintReadWrite, Category=Spring)
+	bool bLimitDisplacement;
+
+	/** If bLimitDisplacement is true, this indicates how long a bone can stretch beyond its length in the ref-pose. */
+	UPROPERTY(EditAnywhere, BlueprintReadWrite, Category=Spring, meta=(EditCondition="bLimitDisplacement"))
+	float MaxDisplacement;
+
+	/** Stiffness of spring */
+	UPROPERTY(EditAnywhere, BlueprintReadWrite, Category=Spring)
+	float SpringStiffness;
+
+	/** Damping of spring */
+	UPROPERTY(EditAnywhere, BlueprintReadWrite, Category=Spring)
+	float SpringDamping;
+
+	/** If spring stretches more than this, reset it. Useful for catching teleports etc */
+	UPROPERTY(EditAnywhere, BlueprintReadWrite, Category=Spring)
+	float ErrorResetThresh;
+
+	/** If true, Z position is always correct, no spring applied */
+	UPROPERTY()
+	bool bNoZSpring_DEPRECATED;
+
+	/** If true take the spring calculation for translation in X */
+	UPROPERTY(EditAnywhere, BlueprintReadWrite, Category=FilterChannels)
+	bool bTranslateX;
+
+	/** If true take the spring calculation for translation in Y */
+	UPROPERTY(EditAnywhere, BlueprintReadWrite, Category=FilterChannels)
+	bool bTranslateY;
+
+	/** If true take the spring calculation for translation in Z */
+	UPROPERTY(EditAnywhere, BlueprintReadWrite, Category=FilterChannels)
+	bool bTranslateZ;
+
+	/** If true take the spring calculation for rotation in X */
+	UPROPERTY(EditAnywhere, BlueprintReadWrite, Category=FilterChannels)
+	bool bRotateX;
+
+	/** If true take the spring calculation for rotation in Y */
+	UPROPERTY(EditAnywhere, BlueprintReadWrite, Category=FilterChannels)
+	bool bRotateY;
+
+	/** If true take the spring calculation for rotation in Z */
+	UPROPERTY(EditAnywhere, BlueprintReadWrite, Category=FilterChannels)
+	bool bRotateZ;
+
+	/** Internal use - Amount of time we need to simulate. */
+	float RemainingTime;
+
+	/** Internal use - Current timestep */
+	float FixedTimeStep;
+
+	/** Did we have a non-zero ControlStrength last frame. */
+	bool bHadValidStrength;
+
+	/** World-space location of the bone. */
+	FVector BoneLocation;
+
+	/** World-space velocity of the bone. */
+	FVector BoneVelocity;
+
+public:
+	FAnimNode_SpringBone();
+
+	// FAnimNode_Base interface
+	virtual void Initialize(const FAnimationInitializeContext& Context) override;
+	virtual void CacheBones(const FAnimationCacheBonesContext& Context) override;
+	virtual void Update(const FAnimationUpdateContext& Context) override;
+	virtual void GatherDebugData(FNodeDebugData& DebugData) override;
+	// End of FAnimNode_Base interface
+
+	// FAnimNode_SkeletalControlBase interface
+	virtual void EvaluateBoneTransforms(USkeletalMeshComponent* SkelComp, FCSPose<FCompactPose>& MeshBases, TArray<FBoneTransform>& OutBoneTransforms) override;
+	virtual bool IsValidToEvaluate(const USkeleton* Skeleton, const FBoneContainer& RequiredBones) override;
+	// End of FAnimNode_SkeletalControlBase interface
+
+private:
+	// FAnimNode_SkeletalControlBase interface
+	virtual void InitializeBoneReferences(const FBoneContainer& RequiredBones) override;
+	// End of FAnimNode_SkeletalControlBase interface
+};
diff --git a/Engine/Source/Runtime/AnimGraphRuntime/Public/BoneControllers/AnimNode_Trail.h b/Engine/Source/Runtime/AnimGraphRuntime/Public/BoneControllers/AnimNode_Trail.h
new file mode 100644
--- /dev/null
+++ b/Engine/Source/Runtime/AnimGraphRuntime/Public/BoneControllers/AnimNode_Trail.h
@@ -0,0 +1,84 @@
+// Copyright 1998-2015 Epic Games, Inc. All Rights Reserved.
+
+#pragma once
+
+#include "AnimNode_SkeletalControlBase.h"
+#include "AnimNode_Trail.generated.h"
+
+/**
+ * Trail Controller
+ */
+
+USTRUCT()
+struct ANIMGRAPHRUNTIME_API FAnimNode_Trail : public FAnimNode_SkeletalControlBase
+{
+	GENERATED_USTRUCT_BODY()
+
+	/** Reference to the active bone in the hierarchy to modify. */
+	UPROPERTY(EditAnywhere, BlueprintReadWrite, Category=Trail)
+	FBoneReference TrailBone;
+
+	/** Number of bones above the active one in the hierarchy to modify. ChainLength should be at least 2. */
+	UPROPERTY(EditAnywhere, BlueprintReadWrite, Category = Trail, meta = (ClampMin = "2", UIMin = "2"))
+	int32	ChainLength;
+
+	/** Axis of the bones to point along trail. */
+	UPROPERTY(EditAnywhere, BlueprintReadWrite, Category=Trail)
+	TEnumAsByte<EAxis::Type>	ChainBoneAxis;
+
+	/** Invert the direction specified in ChainBoneAxis. */
+	UPROPERTY(EditAnywhere, BlueprintReadWrite, Category=Trail)
+	bool	bInvertChainBoneAxis;
+
+	/** Limit the amount that a bone can stretch from its ref-pose length. */
+	UPROPERTY(EditAnywhere, BlueprintReadWrite, Category=Trail)
+	bool	bLimitStretch;
+
+	/** How quickly we 'relax' the bones to their animated positions. */
+	UPROPERTY(EditAnywhere, BlueprintReadWrite, Category=Trail)
+	float	TrailRelaxation;
+
+	/** If bLimitStretch is true, this indicates how long a bone can stretch beyond its length in the ref-pose. */
+	UPROPERTY(EditAnywhere, BlueprintReadWrite, Category=Trail)
+	float	StretchLimit;
+
+	/** 'Fake' velocity applied to bones. */
+	UPROPERTY(EditAnywhere, BlueprintReadWrite, Category=Trail)
+	FVector	FakeVelocity;
+
+	/** Whether 'fake' velocity should be applied in actor or world space. */
+	UPROPERTY(EditAnywhere, BlueprintReadWrite, Category=Trail)
+	bool	bActorSpaceFakeVel;
+
+	/** Internal use - we need the timestep to do the relaxation in CalculateNewBoneTransforms. */
+	float	ThisTimstep;
+
+	/** Did we have a non-zero ControlStrength last frame. */
+	bool	bHadValidStrength;
+
+	/** Component-space locations of the bones from last frame. Each frame these are moved towards their 'animated' locations. */
+	TArray<FVector>	TrailBoneLocations;
+
+	/** LocalToWorld used last frame, used for building transform between frames. */
+	FMatrix		OldLocalToWorld;
+
+
+	FAnimNode_Trail();
+
+	// FAnimNode_Base interface
+	virtual void Update(const FAnimationUpdateContext& Context) override;
+	virtual void GatherDebugData(FNodeDebugData& DebugData) override;
+	// End of FAnimNode_Base interface
+
+	// FAnimNode_SkeletalControlBase interface
+	virtual void EvaluateBoneTransforms(USkeletalMeshComponent* SkelComp, FCSPose<FCompactPose>& MeshBases, TArray<FBoneTransform>& OutBoneTransforms) override;
+	virtual bool IsValidToEvaluate(const USkeleton* Skeleton, const FBoneContainer& RequiredBones) override;
+	// End of FAnimNode_SkeletalControlBase interface
+
+private:
+	// FAnimNode_SkeletalControlBase interface
+	virtual void InitializeBoneReferences(const FBoneContainer& RequiredBones) override;
+	// End of FAnimNode_SkeletalControlBase interface
+
+	FVector GetAlignVector(EAxis::Type AxisOption, bool bInvert);
+};
\ No newline at end of file
diff --git a/Engine/Source/Runtime/AnimGraphRuntime/Public/BoneControllers/AnimNode_TwoBoneIK.h b/Engine/Source/Runtime/AnimGraphRuntime/Public/BoneControllers/AnimNode_TwoBoneIK.h
new file mode 100644
--- /dev/null
+++ b/Engine/Source/Runtime/AnimGraphRuntime/Public/BoneControllers/AnimNode_TwoBoneIK.h
@@ -0,0 +1,76 @@
+// Copyright 1998-2015 Epic Games, Inc. All Rights Reserved.
+
+#pragma once
+
+#include "AnimNode_SkeletalControlBase.h"
+#include "AnimNode_TwoBoneIK.generated.h"
+
+/**
+ * Simple 2 Bone IK Controller.
+ */
+
+USTRUCT()
+struct ANIMGRAPHRUNTIME_API FAnimNode_TwoBoneIK : public FAnimNode_SkeletalControlBase
+{
+	GENERATED_USTRUCT_BODY()
+	
+	/** Name of bone to control. This is the main bone chain to modify from. **/
+	UPROPERTY(EditAnywhere, BlueprintReadWrite, Category=IK)
+	FBoneReference IKBone;
+
+	/** Effector Location. Target Location to reach. */
+	UPROPERTY(EditAnywhere, BlueprintReadWrite, Category=EndEffector, meta=(PinShownByDefault))
+	FVector EffectorLocation;
+
+	/** Joint Target Location. Location used to orient Joint bone. **/
+	UPROPERTY(EditAnywhere, BlueprintReadWrite, Category = JointTarget, meta=(PinShownByDefault))
+	FVector JointTargetLocation;
+
+	/** Limits to use if stretching is allowed. */
+	UPROPERTY(EditAnywhere, BlueprintReadWrite, Category=IK)
+	FVector2D StretchLimits;
+
+	/** If EffectorLocationSpace is a bone, this is the bone to use. **/
+	UPROPERTY(EditAnywhere, BlueprintReadWrite, Category=EndEffector)
+	FName EffectorSpaceBoneName;
+
+	/** Set end bone to use End Effector rotation */
+	UPROPERTY(EditAnywhere, BlueprintReadWrite, Category=EndEffector)
+	uint32 bTakeRotationFromEffectorSpace:1;
+
+	/** Keep local rotation of end bone */
+	UPROPERTY(EditAnywhere, BlueprintReadWrite, Category=EndEffector)
+	uint32 bMaintainEffectorRelRot:1;
+
+	/** Should stretching be allowed, to be prevent over extension */
+	UPROPERTY(EditAnywhere, BlueprintReadWrite, Category=IK)
+	uint32 bAllowStretching:1;
+	
+	/** Reference frame of Effector Location. */
+	UPROPERTY(EditAnywhere, BlueprintReadWrite, Category = IK)
+	TEnumAsByte<enum EBoneControlSpace> EffectorLocationSpace;
+
+	/** Reference frame of Joint Target Location. */
+	UPROPERTY(EditAnywhere, BlueprintReadWrite, Category=JointTarget)
+	TEnumAsByte<enum EBoneControlSpace> JointTargetLocationSpace;
+
+	/** If JointTargetSpaceBoneName is a bone, this is the bone to use. **/
+	UPROPERTY(EditAnywhere, BlueprintReadWrite, Category=JointTarget)
+	FName JointTargetSpaceBoneName;
+
+	FAnimNode_TwoBoneIK();
+
+	// FAnimNode_Base interface
+	virtual void GatherDebugData(FNodeDebugData& DebugData) override;
+	// End of FAnimNode_Base interface
+
+	// FAnimNode_SkeletalControlBase interface
+	virtual void EvaluateBoneTransforms(USkeletalMeshComponent* SkelComp, FCSPose<FCompactPose>& MeshBases, TArray<FBoneTransform>& OutBoneTransforms) override;
+	virtual bool IsValidToEvaluate(const USkeleton* Skeleton, const FBoneContainer& RequiredBones) override;
+	// End of FAnimNode_SkeletalControlBase interface
+
+private:
+	// FAnimNode_SkeletalControlBase interface
+	virtual void InitializeBoneReferences(const FBoneContainer& RequiredBones) override;
+	// End of FAnimNode_SkeletalControlBase interface
+};
\ No newline at end of file
diff --git a/Engine/Source/Runtime/AnimGraphRuntime/Public/BoneControllers/AnimNode_WheelHandler.h b/Engine/Source/Runtime/AnimGraphRuntime/Public/BoneControllers/AnimNode_WheelHandler.h
new file mode 100644
--- /dev/null
+++ b/Engine/Source/Runtime/AnimGraphRuntime/Public/BoneControllers/AnimNode_WheelHandler.h
@@ -0,0 +1,47 @@
+// Copyright 1998-2015 Epic Games, Inc. All Rights Reserved.
+
+#pragma once
+
+#include "AnimNode_SkeletalControlBase.h"
+#include "AnimNode_WheelHandler.generated.h"
+
+struct FWheelSimulator
+{
+	int32					WheelIndex;
+	FBoneReference			BoneReference;
+	FRotator				RotOffset;
+	FVector					LocOffset;
+};
+
+/**
+ *	Simple controller that replaces or adds to the translation/rotation of a single bone.
+ */
+USTRUCT()
+struct ANIMGRAPHRUNTIME_API FAnimNode_WheelHandler : public FAnimNode_SkeletalControlBase
+{
+	GENERATED_USTRUCT_BODY()
+
+	/** Current Asset being played **/
+	UPROPERTY(transient)
+	class UWheeledVehicleMovementComponent* VehicleSimComponent;
+
+	TArray<FWheelSimulator>			WheelSimulators;
+
+	FAnimNode_WheelHandler();
+
+	// FAnimNode_Base interface
+	virtual void GatherDebugData(FNodeDebugData& DebugData) override;
+	// End of FAnimNode_Base interface
+
+	// FAnimNode_SkeletalControlBase interface
+	virtual void EvaluateBoneTransforms(USkeletalMeshComponent* SkelComp, FCSPose<FCompactPose>& MeshBases, TArray<FBoneTransform>& OutBoneTransforms) override;
+	virtual bool IsValidToEvaluate(const USkeleton* Skeleton, const FBoneContainer& RequiredBones) override;
+	virtual void Update(const FAnimationUpdateContext& Context) override;
+	virtual void Initialize(const FAnimationInitializeContext& Context) override;
+	// End of FAnimNode_SkeletalControlBase interface
+
+private:
+	// FAnimNode_SkeletalControlBase interface
+	virtual void InitializeBoneReferences(const FBoneContainer& RequiredBones) override;
+	// End of FAnimNode_SkeletalControlBase interface
+};
diff --git a/Engine/Source/Runtime/Engine/Classes/Animation/AnimNodeBase.h b/Engine/Source/Runtime/Engine/Classes/Animation/AnimNodeBase.h
--- a/Engine/Source/Runtime/Engine/Classes/Animation/AnimNodeBase.h
+++ b/Engine/Source/Runtime/Engine/Classes/Animation/AnimNodeBase.h
@@ -176,7 +176,7 @@ public:
 	bool IsNormalized() const;
 };
 
-struct FNodeDebugData
+struct ENGINE_API FNodeDebugData
 {
 private:
 	struct DebugItem
@@ -187,14 +187,14 @@ private:
 		FString DebugData;
 
 		// Whether we are supplying a pose instead of modifying one (e.g. an playing animation)
-		bool	bPoseSource;
+		bool bPoseSource;
 
 		// Nodes that we are connected to
 		TArray<FNodeDebugData> ChildNodeChain;
 	};
 
 	// This nodes final contribution weight (based on its own weight and the weight of its parents)
-	float			AbsoluteWeight;
+	float AbsoluteWeight;
 
 	// Nodes that we are dependent on
 	TArray<DebugItem> NodeChain;
@@ -204,10 +204,10 @@ public:
 	{
 		FFlattenedDebugData(FString Line, float AbsWeight, int32 InIndent, int32 InChainID, bool bInPoseSource) : DebugLine(Line), AbsoluteWeight(AbsWeight), Indent(InIndent), ChainID(InChainID), bPoseSource(bInPoseSource){}
 		FString DebugLine;
-		float	AbsoluteWeight;
-		int32	Indent;
-		int32	ChainID;
-		bool	bPoseSource;
+		float AbsoluteWeight;
+		int32 Indent;
+		int32 ChainID;
+		bool bPoseSource;
 
 		bool IsOnActiveBranch() { return AbsoluteWeight > ZERO_ANIMWEIGHT_THRESH; }
 	};
@@ -215,8 +215,8 @@ public:
 	FNodeDebugData(const class UAnimInstance* InAnimInstance) : AbsoluteWeight(1.f), AnimInstance(InAnimInstance) {}
 	FNodeDebugData(const class UAnimInstance* InAnimInstance, const float AbsWeight) : AbsoluteWeight(AbsWeight), AnimInstance(InAnimInstance) {}
 
-	void			AddDebugItem(FString DebugData, bool bPoseSource = false);
-	FNodeDebugData&	BranchFlow(float BranchWeight);
+	void AddDebugItem(FString DebugData, bool bPoseSource = false);
+	FNodeDebugData& BranchFlow(float BranchWeight);
 
 	template<class Type>
 	FString GetNodeName(Type* Node)
@@ -235,7 +235,7 @@ public:
 	}
 
 	// Anim instance that we are generating debug data for
-	const class UAnimInstance* AnimInstance;
+	const UAnimInstance* AnimInstance;
 };
 
 // The display mode of editable values on an animation node
diff --git a/Engine/Source/Runtime/Engine/Classes/Animation/BoneControllers/AnimNode_BoneDrivenController.h b/Engine/Source/Runtime/Engine/Classes/Animation/BoneControllers/AnimNode_BoneDrivenController.h
deleted file mode 100644
diff --git a/Engine/Source/Runtime/Engine/Classes/Animation/BoneControllers/AnimNode_CopyBone.h b/Engine/Source/Runtime/Engine/Classes/Animation/BoneControllers/AnimNode_CopyBone.h
deleted file mode 100644
diff --git a/Engine/Source/Runtime/Engine/Classes/Animation/BoneControllers/AnimNode_Fabrik.h b/Engine/Source/Runtime/Engine/Classes/Animation/BoneControllers/AnimNode_Fabrik.h
deleted file mode 100644
diff --git a/Engine/Source/Runtime/Engine/Classes/Animation/BoneControllers/AnimNode_HandIKRetargeting.h b/Engine/Source/Runtime/Engine/Classes/Animation/BoneControllers/AnimNode_HandIKRetargeting.h
deleted file mode 100644
diff --git a/Engine/Source/Runtime/Engine/Classes/Animation/BoneControllers/AnimNode_LookAt.h b/Engine/Source/Runtime/Engine/Classes/Animation/BoneControllers/AnimNode_LookAt.h
deleted file mode 100644
diff --git a/Engine/Source/Runtime/Engine/Classes/Animation/BoneControllers/AnimNode_ModifyBone.h b/Engine/Source/Runtime/Engine/Classes/Animation/BoneControllers/AnimNode_ModifyBone.h
deleted file mode 100644
diff --git a/Engine/Source/Runtime/Engine/Classes/Animation/BoneControllers/AnimNode_RotationMultiplier.h b/Engine/Source/Runtime/Engine/Classes/Animation/BoneControllers/AnimNode_RotationMultiplier.h
deleted file mode 100644
diff --git a/Engine/Source/Runtime/Engine/Classes/Animation/BoneControllers/AnimNode_SkeletalControlBase.h b/Engine/Source/Runtime/Engine/Classes/Animation/BoneControllers/AnimNode_SkeletalControlBase.h
deleted file mode 100644
diff --git a/Engine/Source/Runtime/Engine/Classes/Animation/BoneControllers/AnimNode_SpringBone.h b/Engine/Source/Runtime/Engine/Classes/Animation/BoneControllers/AnimNode_SpringBone.h
deleted file mode 100644
diff --git a/Engine/Source/Runtime/Engine/Classes/Animation/BoneControllers/AnimNode_Trail.h b/Engine/Source/Runtime/Engine/Classes/Animation/BoneControllers/AnimNode_Trail.h
deleted file mode 100644
diff --git a/Engine/Source/Runtime/Engine/Classes/Animation/BoneControllers/AnimNode_TwoBoneIK.h b/Engine/Source/Runtime/Engine/Classes/Animation/BoneControllers/AnimNode_TwoBoneIK.h
deleted file mode 100644
diff --git a/Engine/Source/Runtime/Engine/Classes/Animation/BoneControllers/AnimNode_WheelHandler.h b/Engine/Source/Runtime/Engine/Classes/Animation/BoneControllers/AnimNode_WheelHandler.h
deleted file mode 100644
diff --git a/Engine/Source/Runtime/Engine/Engine.Build.cs b/Engine/Source/Runtime/Engine/Engine.Build.cs
--- a/Engine/Source/Runtime/Engine/Engine.Build.cs
+++ b/Engine/Source/Runtime/Engine/Engine.Build.cs
@@ -78,6 +78,10 @@ public class Engine : ModuleRules
         CircularlyReferencedDependentModules.Add("UMG");
 		CircularlyReferencedDependentModules.Add("Niagara");
 
+		// The AnimGraphRuntime module is not needed by Engine proper, but it is loaded in LaunchEngineLoop.cpp,
+		// and needs to be listed in an always-included module in order to be compiled into standalone games
+		DynamicallyLoadedModuleNames.Add("AnimGraphRuntime");
+
 		DynamicallyLoadedModuleNames.AddRange(
 			new string[]
 			{
diff --git a/Engine/Source/Runtime/Engine/Private/Animation/AnimNodeRotationMultiplier.cpp b/Engine/Source/Runtime/Engine/Private/Animation/AnimNodeRotationMultiplier.cpp
deleted file mode 100644
diff --git a/Engine/Source/Runtime/Engine/Private/Animation/AnimNodeSkeletalControlBase.cpp b/Engine/Source/Runtime/Engine/Private/Animation/AnimNodeSkeletalControlBase.cpp
deleted file mode 100644
diff --git a/Engine/Source/Runtime/Engine/Private/Animation/AnimNode_BoneDrivenController.cpp b/Engine/Source/Runtime/Engine/Private/Animation/AnimNode_BoneDrivenController.cpp
deleted file mode 100644
diff --git a/Engine/Source/Runtime/Engine/Private/Animation/AnimNode_CopyBone.cpp b/Engine/Source/Runtime/Engine/Private/Animation/AnimNode_CopyBone.cpp
deleted file mode 100644
diff --git a/Engine/Source/Runtime/Engine/Private/Animation/AnimNode_Fabrik.cpp b/Engine/Source/Runtime/Engine/Private/Animation/AnimNode_Fabrik.cpp
deleted file mode 100644
diff --git a/Engine/Source/Runtime/Engine/Private/Animation/AnimNode_HandIKRetargeting.cpp b/Engine/Source/Runtime/Engine/Private/Animation/AnimNode_HandIKRetargeting.cpp
deleted file mode 100644
diff --git a/Engine/Source/Runtime/Engine/Private/Animation/AnimNode_LookAt.cpp b/Engine/Source/Runtime/Engine/Private/Animation/AnimNode_LookAt.cpp
deleted file mode 100644
diff --git a/Engine/Source/Runtime/Engine/Private/Animation/AnimNode_ModifyBone.cpp b/Engine/Source/Runtime/Engine/Private/Animation/AnimNode_ModifyBone.cpp
deleted file mode 100644
diff --git a/Engine/Source/Runtime/Engine/Private/Animation/AnimNode_SpringBone.cpp b/Engine/Source/Runtime/Engine/Private/Animation/AnimNode_SpringBone.cpp
deleted file mode 100644
diff --git a/Engine/Source/Runtime/Engine/Private/Animation/AnimNode_Trail.cpp b/Engine/Source/Runtime/Engine/Private/Animation/AnimNode_Trail.cpp
deleted file mode 100644
diff --git a/Engine/Source/Runtime/Engine/Private/Animation/AnimNode_TwoBoneIK.cpp b/Engine/Source/Runtime/Engine/Private/Animation/AnimNode_TwoBoneIK.cpp
deleted file mode 100644
diff --git a/Engine/Source/Runtime/Engine/Private/Animation/AnimNode_WheelHandler.cpp b/Engine/Source/Runtime/Engine/Private/Animation/AnimNode_WheelHandler.cpp
deleted file mode 100644
diff --git a/Engine/Source/Runtime/Engine/Public/BoneContainer.h b/Engine/Source/Runtime/Engine/Public/BoneContainer.h
--- a/Engine/Source/Runtime/Engine/Public/BoneContainer.h
+++ b/Engine/Source/Runtime/Engine/Public/BoneContainer.h
@@ -13,7 +13,7 @@ class USkeleton;
 * - BoneIndicesArray: Array of RequiredBoneIndices for Current Asset. In increasing order. Mapping to current Array of Transforms (Pose).
 * - BoneSwitchArray: Size of current Skeleton. true if Bone is contained in RequiredBones array, false otherwise.
 **/
-struct FBoneContainer
+struct ENGINE_API FBoneContainer
 {
 private:
 	/** Array of RequiredBonesIndices. In increasing order. */
@@ -176,13 +176,13 @@ public:
 	}
 
 	/** Get BoneIndex for BoneName for current Asset. */
-	ENGINE_API int32 GetPoseBoneIndexForBoneName(const FName& BoneName) const;
+	int32 GetPoseBoneIndexForBoneName(const FName& BoneName) const;
 
 	/** Get ParentBoneIndex for current Asset. */
-	ENGINE_API int32 GetParentBoneIndex(const int32& BoneIndex) const;
+	int32 GetParentBoneIndex(const int32& BoneIndex) const;
 
 	/** Get ParentBoneIndex for current Asset. */
-	ENGINE_API FCompactPoseBoneIndex GetParentBoneIndex(const FCompactPoseBoneIndex& BoneIndex) const;
+	FCompactPoseBoneIndex GetParentBoneIndex(const FCompactPoseBoneIndex& BoneIndex) const;
 
 	/** Get Depth between bones for current asset. */
 	int32 GetDepthBetweenBones(const int32& BoneIndex, const int32& ParentBoneIndex) const;
@@ -256,7 +256,7 @@ public:
 
 private:
 	/** Initialize FBoneContainer. */
-	ENGINE_API void Initialize();
+	void Initialize();
 
 	/** Cache remapping data if current Asset is a SkeletalMesh, with all compatible Skeletons. */
 	void RemapFromSkelMesh(USkeletalMesh const & SourceSkeletalMesh, USkeleton& TargetSkeleton);
diff --git a/Engine/Source/Runtime/Launch/Private/LaunchEngineLoop.cpp b/Engine/Source/Runtime/Launch/Private/LaunchEngineLoop.cpp
--- a/Engine/Source/Runtime/Launch/Private/LaunchEngineLoop.cpp
+++ b/Engine/Source/Runtime/Launch/Private/LaunchEngineLoop.cpp
@@ -1786,6 +1786,8 @@ void FEngineLoop::LoadPreInitModules()
 
 	FModuleManager::Get().LoadModule(TEXT("Renderer"));
 
+	FModuleManager::Get().LoadModule(TEXT("AnimGraphRuntime"));
+
 	FPlatformMisc::LoadPreInitModules();
 	
 #if !UE_SERVER
