Animation: Added missing Build.cs file

diff --git a/Engine/Source/Runtime/AnimGraphRuntime/AnimGraphRuntime.Build.cs b/Engine/Source/Runtime/AnimGraphRuntime/AnimGraphRuntime.Build.cs
new file mode 100644
--- /dev/null
+++ b/Engine/Source/Runtime/AnimGraphRuntime/AnimGraphRuntime.Build.cs
@@ -0,0 +1,20 @@
+// Copyright 1998-2015 Epic Games, Inc. All Rights Reserved.
+
+using UnrealBuildTool;
+using System.IO;
+
+public class AnimGraphRuntime : ModuleRules
+{
+	public AnimGraphRuntime(TargetInfo Target)
+	{
+		PrivateIncludePaths.Add("Runtime/AnimGraphRuntime/Private");
+
+		PublicDependencyModuleNames.AddRange(
+			new string[] { 
+				"Core", 
+				"CoreUObject", 
+				"Engine"
+			}
+		);
+	}
+}
