Redoing a change to an enum that got wiped when the file was moved.

diff --git a/Engine/Source/Runtime/AnimGraphRuntime/Public/BoneControllers/AnimNode_ModifyBone.h b/Engine/Source/Runtime/AnimGraphRuntime/Public/BoneControllers/AnimNode_ModifyBone.h
--- a/Engine/Source/Runtime/AnimGraphRuntime/Public/BoneControllers/AnimNode_ModifyBone.h
+++ b/Engine/Source/Runtime/AnimGraphRuntime/Public/BoneControllers/AnimNode_ModifyBone.h
@@ -8,13 +8,13 @@
 UENUM()
 enum EBoneModificationMode
 {
-	// The modifier ignores this channel (keeps the existing bone translation, rotation, or scale)
+	/** The modifier ignores this channel (keeps the existing bone translation, rotation, or scale). */
 	BMM_Ignore,
 
-	// The modifier replaces the existing translation, rotation, or scale
+	/** The modifier replaces the existing translation, rotation, or scale. */
 	BMM_Replace,
 
-	// The modifier adds to the existing translation, rotation, or scale
+	/** The modifier adds to the existing translation, rotation, or scale. */
 	BMM_Additive
 };
 
