- Added Look up axis

diff --git a/Engine/Source/Runtime/AnimGraphRuntime/Private/BoneControllers/AnimNode_LookAt.cpp b/Engine/Source/Runtime/AnimGraphRuntime/Private/BoneControllers/AnimNode_LookAt.cpp
--- a/Engine/Source/Runtime/AnimGraphRuntime/Private/BoneControllers/AnimNode_LookAt.cpp
+++ b/Engine/Source/Runtime/AnimGraphRuntime/Private/BoneControllers/AnimNode_LookAt.cpp
@@ -4,6 +4,23 @@
 #include "BoneControllers/AnimNode_LookAt.h"
 #include "AnimationRuntime.h"
 
+FVector FAnimNode_LookAt::GetAlignVector(const FTransform& Transform, EAxisOption::Type AxisOption)
+{
+	switch(AxisOption)
+	{
+	case EAxisOption::X:
+	case EAxisOption::X_Neg:
+	return Transform.GetUnitAxis(EAxis::X);
+	case EAxisOption::Y:
+	case EAxisOption::Y_Neg:
+	return Transform.GetUnitAxis(EAxis::Y);
+	case EAxisOption::Z:
+	case EAxisOption::Z_Neg:
+	return Transform.GetUnitAxis(EAxis::Z);
+	}
+
+	return FVector(1.f, 0.f, 0.f);
+}
 /////////////////////////////////////////////////////
 // FAnimNode_LookAt
 
@@ -140,8 +157,22 @@ void FAnimNode_LookAt::EvaluateBoneTransforms(USkeletalMeshComponent* SkelComp,
 		}
 	}
 
-	// get delta rotation
-	FQuat DeltaRot = FQuat::FindBetween(LookAtVector, ToTarget);
+	FQuat DeltaRot;
+	// if want to use look up, 
+	if (bUseLookUpAxis)
+	{
+		// find look up vector in local space
+		FVector LookUpVector = GetAlignVector(ComponentBoneTransform, LookUpAxis);
+		// project target to the plane
+		FVector NewTarget = FVector::VectorPlaneProject(ToTarget, LookUpVector);
+		NewTarget.Normalize();
+		DeltaRot = FQuat::FindBetween(LookAtVector, NewTarget);
+	}
+	else
+	{
+		DeltaRot = FQuat::FindBetween(LookAtVector, ToTarget);
+	}
+
 	// transform current rotation to delta rotation
 	FQuat CurrentRot = ComponentBoneTransform.GetRotation();
 	FQuat NewRotation = DeltaRot * CurrentRot;
@@ -166,24 +197,6 @@ void FAnimNode_LookAt::InitializeBoneReferences(const FBoneContainer& RequiredBo
 	LookAtBone.Initialize(RequiredBones);
 }
 
-FVector FAnimNode_LookAt::GetAlignVector(const FTransform& Transform, EAxisOption::Type AxisOption)
-{
-	switch (AxisOption)
-	{
-	case EAxisOption::X:
-	case EAxisOption::X_Neg:
-		return Transform.GetUnitAxis(EAxis::X);
-	case EAxisOption::Y:
-	case EAxisOption::Y_Neg:
-		return Transform.GetUnitAxis(EAxis::Y);
-	case EAxisOption::Z:
-	case EAxisOption::Z_Neg:
-		return Transform.GetUnitAxis(EAxis::Z);
-	}
-
-	return FVector(1.f, 0.f, 0.f);
-}
-
 void FAnimNode_LookAt::Update(const FAnimationUpdateContext& Context)
 {
 	FAnimNode_SkeletalControlBase::Update(Context);
diff --git a/Engine/Source/Runtime/AnimGraphRuntime/Public/BoneControllers/AnimNode_LookAt.h b/Engine/Source/Runtime/AnimGraphRuntime/Public/BoneControllers/AnimNode_LookAt.h
--- a/Engine/Source/Runtime/AnimGraphRuntime/Public/BoneControllers/AnimNode_LookAt.h
+++ b/Engine/Source/Runtime/AnimGraphRuntime/Public/BoneControllers/AnimNode_LookAt.h
@@ -61,6 +61,14 @@ struct ANIMGRAPHRUNTIME_API FAnimNode_LookAt : public FAnimNode_SkeletalControlB
 	UPROPERTY(EditAnywhere, BlueprintReadWrite, Category=SkeletalControl) 
 	TEnumAsByte<EAxisOption::Type>	LookAtAxis;
 
+	/** Whether or not to use Look up axis */
+	UPROPERTY(EditAnywhere, BlueprintReadWrite, Category=SkeletalControl)
+	bool bUseLookUpAxis;
+
+	/** Look up axis in local space */
+	UPROPERTY(EditAnywhere, BlueprintReadWrite, Category=SkeletalControl)
+	TEnumAsByte<EAxisOption::Type>	LookUpAxis;
+
 	/** Look at Clamp value in degree - if you're look at axis is Z, only X, Y degree of clamp will be used*/
 	UPROPERTY(EditAnywhere, BlueprintReadWrite, Category=SkeletalControl, meta=(PinHiddenByDefault))
 	float LookAtClamp;
