Minor optimization from Github submission #1133

diff --git a/Engine/Source/Runtime/AnimGraphRuntime/Private/BoneControllers/AnimNode_Fabrik.cpp b/Engine/Source/Runtime/AnimGraphRuntime/Private/BoneControllers/AnimNode_Fabrik.cpp
--- a/Engine/Source/Runtime/AnimGraphRuntime/Private/BoneControllers/AnimNode_Fabrik.cpp
+++ b/Engine/Source/Runtime/AnimGraphRuntime/Private/BoneControllers/AnimNode_Fabrik.cpp
@@ -105,12 +105,12 @@ void FAnimNode_Fabrik::EvaluateBoneTransforms(USkeletalMeshComponent* SkelComp,
 	}
 
 	bool bBoneLocationUpdated = false;
-	float const RootToTargetDist = FVector::Dist(Chain[0].Position, CSEffectorLocation);
+	float const RootToTargetDistSq = FVector::DistSquared(Chain[0].Position, CSEffectorLocation);
 	int32 const NumChainLinks = Chain.Num();
 
 	// FABRIK algorithm - bone translation calculation
 	// If the effector is further away than the distance from root to tip, simply move all bones in a line from root to effector location
-	if (RootToTargetDist > MaximumReach)
+	if (RootToTargetDistSq > FMath::Square(MaximumReach))
 	{
 		for (int32 LinkIndex = 1; LinkIndex < NumChainLinks; LinkIndex++)
 		{
