Animation: Renamed AnimNodeSkeletalControlBase.cpp to AnimNode_SkeletalControlBase.cpp, so it matches the header and class name

diff --git a/Engine/Source/Runtime/AnimGraphRuntime/Private/BoneControllers/AnimNodeSkeletalControlBase.cpp b/Engine/Source/Runtime/AnimGraphRuntime/Private/BoneControllers/AnimNodeSkeletalControlBase.cpp
deleted file mode 100644
diff --git a/Engine/Source/Runtime/AnimGraphRuntime/Private/BoneControllers/AnimNode_SkeletalControlBase.cpp b/Engine/Source/Runtime/AnimGraphRuntime/Private/BoneControllers/AnimNode_SkeletalControlBase.cpp
new file mode 100644
--- /dev/null
+++ b/Engine/Source/Runtime/AnimGraphRuntime/Private/BoneControllers/AnimNode_SkeletalControlBase.cpp
@@ -0,0 +1,74 @@
+// Copyright 1998-2015 Epic Games, Inc. All Rights Reserved.
+
+#include "AnimGraphRuntimePrivatePCH.h"
+#include "BoneControllers/AnimNode_SkeletalControlBase.h"
+#include "AnimationUtils.h"
+#include "AnimationRuntime.h"
+
+/////////////////////////////////////////////////////
+// FAnimNode_SkeletalControlBase
+
+void FAnimNode_SkeletalControlBase::Initialize(const FAnimationInitializeContext& Context)
+{
+	ComponentPose.Initialize(Context);
+}
+
+void FAnimNode_SkeletalControlBase::CacheBones(const FAnimationCacheBonesContext& Context) 
+{
+	InitializeBoneReferences(Context.AnimInstance->RequiredBones);
+	ComponentPose.CacheBones(Context);
+}
+
+void FAnimNode_SkeletalControlBase::Update(const FAnimationUpdateContext& Context)
+{
+	ComponentPose.Update(Context);
+	EvaluateGraphExposedInputs.Execute(Context);
+}
+
+bool ContainsNaN(const TArray<FBoneTransform> & BoneTransforms)
+{
+	for (int32 i = 0; i < BoneTransforms.Num(); ++i)
+	{
+		if (BoneTransforms[i].Transform.ContainsNaN())
+		{
+			return true;
+		}
+	}
+
+	return false;
+}
+
+void FAnimNode_SkeletalControlBase::EvaluateComponentSpace(FComponentSpacePoseContext& Output)
+{
+	// Evaluate the input
+	ComponentPose.EvaluateComponentSpace(Output);
+
+	// Apply the skeletal control if it's valid
+	const float ActualAlpha = AlphaScaleBias.ApplyTo(Alpha);
+	if ((ActualAlpha >= ZERO_ANIMWEIGHT_THRESH) && IsValidToEvaluate(Output.AnimInstance->CurrentSkeleton, Output.AnimInstance->RequiredBones))
+	{
+		USkeletalMeshComponent* Component = Output.AnimInstance->GetSkelMeshComponent();
+
+#if WITH_EDITORONLY_DATA
+		// save current pose before applying skeletal control to compute the exact gizmo location in AnimGraphNode
+		ForwardedPose = Output.Pose;
+#endif // #if WITH_EDITORONLY_DATA
+
+		TArray<FBoneTransform> BoneTransforms;
+		EvaluateBoneTransforms(Component, Output.Pose, BoneTransforms);
+
+		checkSlow(!ContainsNaN(BoneTransforms));
+
+		if (BoneTransforms.Num() > 0)
+		{
+			const float BlendWeight = FMath::Clamp<float>(ActualAlpha, 0.f, 1.f);
+			Output.Pose.LocalBlendCSBoneTransforms(BoneTransforms, BlendWeight);
+		}
+	}
+}
+
+void FAnimNode_SkeletalControlBase::AddDebugNodeData(FString& OutDebugData)
+{
+	const float ActualAlpha = AlphaScaleBias.ApplyTo(Alpha);
+	OutDebugData += FString::Printf(TEXT("Alpha: %.1f%%"), ActualAlpha*100.f);
+}
\ No newline at end of file
