Animation: Added friendly display names to EBoneModificationMode enumerants

diff --git a/Engine/Source/Runtime/AnimGraphRuntime/Public/BoneControllers/AnimNode_ModifyBone.h b/Engine/Source/Runtime/AnimGraphRuntime/Public/BoneControllers/AnimNode_ModifyBone.h
--- a/Engine/Source/Runtime/AnimGraphRuntime/Public/BoneControllers/AnimNode_ModifyBone.h
+++ b/Engine/Source/Runtime/AnimGraphRuntime/Public/BoneControllers/AnimNode_ModifyBone.h
@@ -9,13 +9,13 @@ UENUM()
 enum EBoneModificationMode
 {
 	/** The modifier ignores this channel (keeps the existing bone translation, rotation, or scale). */
-	BMM_Ignore,
+	BMM_Ignore UMETA(DisplayName = "Ignore"),
 
 	/** The modifier replaces the existing translation, rotation, or scale. */
-	BMM_Replace,
+	BMM_Replace UMETA(DisplayName = "Replace Existing"),
 
 	/** The modifier adds to the existing translation, rotation, or scale. */
-	BMM_Additive
+	BMM_Additive UMETA(DisplayName = "Add to Existing")
 };
 
 /**
