Animation: Various improvements to the bone driven controller
- Added support for using a curve to map from input range to output range
- Added compilation warnings when a bone or component still needs to be selected
- Added a customization to hide the range/multipler when using a curve, and gated editing of the range based on bUseRange
- Minor optimization to evaluation speed
- Reorganized the properties into logical groups (Source/Driver, Mapping, Destination/Driven)

diff --git a/Engine/Source/Editor/AnimGraph/AnimGraph.Build.cs b/Engine/Source/Editor/AnimGraph/AnimGraph.Build.cs
--- a/Engine/Source/Editor/AnimGraph/AnimGraph.Build.cs
+++ b/Engine/Source/Editor/AnimGraph/AnimGraph.Build.cs
@@ -25,6 +25,7 @@ public class AnimGraph : ModuleRules
 				"SlateCore",
 				"UnrealEd",
                 "GraphEditor",
+				"PropertyEditor"
 			}
 		);
 
diff --git a/Engine/Source/Editor/AnimGraph/Classes/AnimGraphNode_BoneDrivenController.h b/Engine/Source/Editor/AnimGraph/Classes/AnimGraphNode_BoneDrivenController.h
--- a/Engine/Source/Editor/AnimGraph/Classes/AnimGraphNode_BoneDrivenController.h
+++ b/Engine/Source/Editor/AnimGraph/Classes/AnimGraphNode_BoneDrivenController.h
@@ -8,7 +8,7 @@
 #include "AnimGraphNode_BoneDrivenController.generated.h"
 
 /**
- * This is the 'source version' of of a bone driven controller, which copies part of the state from one bone to another
+ * This is the 'source version' of a bone driven controller, which maps part of the state from one bone to another (e.g., 2 * source.x -> target.z)
  */
 UCLASS()
 class ANIMGRAPH_API UAnimGraphNode_BoneDrivenController : public UAnimGraphNode_SkeletalControlBase
@@ -25,6 +25,11 @@ public:
 	virtual FText GetTooltipText() const override;
 	// End of UEdGraphNode interface
 
+	// UAnimGraphNode_Base interface
+	virtual void ValidateAnimNodeDuringCompilation(USkeleton* ForSkeleton, FCompilerResultsLog& MessageLog) override;
+	virtual void CustomizeDetails(IDetailLayoutBuilder& DetailBuilder) override;
+	// End of UAnimGraphNode_Base interface
+
 	// UAnimGraphNode_SkeletalControlBase interface
 	virtual void Draw(FPrimitiveDrawInterface* PDI, USkeletalMeshComponent* SkelMeshComp) const override;
 	// End of UAnimGraphNode_SkeletalControlBase interface
@@ -35,7 +40,10 @@ protected:
 	virtual FText GetControllerDescription() const override;
 	// End of UAnimGraphNode_SkeletalControlBase protected interface
 
+	// Should non-curve mapping values be shown (multipler, range)?
+	EVisibility AreNonCurveMappingValuesVisible() const;
+
 private:
 	/** Constructing FText strings can be costly, so we cache the node's title */
 	FNodeTitleTextTable CachedNodeTitles;
-};
\ No newline at end of file
+};
diff --git a/Engine/Source/Editor/AnimGraph/Private/AnimGraphNode_BoneDrivenController.cpp b/Engine/Source/Editor/AnimGraph/Private/AnimGraphNode_BoneDrivenController.cpp
--- a/Engine/Source/Editor/AnimGraph/Private/AnimGraphNode_BoneDrivenController.cpp
+++ b/Engine/Source/Editor/AnimGraph/Private/AnimGraphNode_BoneDrivenController.cpp
@@ -2,6 +2,8 @@
 
 #include "AnimGraphPrivatePCH.h"
 #include "AnimGraphNode_BoneDrivenController.h"
+#include "CompilerResultsLog.h"
+#include "PropertyEditing.h"
 
 #define LOCTEXT_NAMESPACE "A3Nodes"
 
@@ -41,7 +43,7 @@ FText UAnimGraphNode_BoneDrivenController::GetNodeTitle(ENodeTitleType::Type Tit
 		}
 
 		// FText::Format() is slow, so we cache this to save on performance
-		CachedNodeTitles.SetCachedTitle(TitleType, FText::Format(LOCTEXT("AnimGraphNode_BoneDrivenController_Title", "{ControllerDesc}{Delim}Driving Bone: {SourceBone}{Delim}Driven Bone: {TargetBone}"), Args), this);
+		CachedNodeTitles.SetCachedTitle(TitleType, FText::Format(LOCTEXT("AnimGraphNode_BoneDrivenController_Title", "{ControllerDesc}{Delim}Driver Bone: {SourceBone}{Delim}Driven Bone: {TargetBone}"), Args), this);
 	}	
 	return CachedNodeTitles[TitleType];
 }
@@ -56,25 +58,24 @@ void UAnimGraphNode_BoneDrivenController::Draw(FPrimitiveDrawInterface* PDI, USk
 	static const float ArrowHeadWidth = 5.0f;
 	static const float ArrowHeadHeight = 8.0f;
 
-	int32 SourceIdx = SkelMeshComp->GetBoneIndex(Node.SourceBone.BoneName);
-	int32 TargetIdx = SkelMeshComp->GetBoneIndex(Node.TargetBone.BoneName);
+	const int32 SourceIdx = SkelMeshComp->GetBoneIndex(Node.SourceBone.BoneName);
+	const int32 TargetIdx = SkelMeshComp->GetBoneIndex(Node.TargetBone.BoneName);
 
-	if(SourceIdx != INDEX_NONE && TargetIdx != INDEX_NONE)
+	if ((SourceIdx != INDEX_NONE) && (TargetIdx != INDEX_NONE))
 	{
-		FTransform SourceTM = SkelMeshComp->GetSpaceBases()[SourceIdx] * SkelMeshComp->ComponentToWorld;
-		FTransform TargetTM = SkelMeshComp->GetSpaceBases()[TargetIdx] * SkelMeshComp->ComponentToWorld;
+		const FTransform SourceTM = SkelMeshComp->GetSpaceBases()[SourceIdx] * SkelMeshComp->ComponentToWorld;
+		const FTransform TargetTM = SkelMeshComp->GetSpaceBases()[TargetIdx] * SkelMeshComp->ComponentToWorld;
 
 		PDI->DrawLine(TargetTM.GetLocation(), SourceTM.GetLocation(), FLinearColor(0.0f, 0.0f, 1.0f), SDPG_Foreground, 0.5f);
 
-		FVector ToTarget = TargetTM.GetTranslation() - SourceTM.GetTranslation();
-		FVector UnitToTarget = ToTarget;
-		UnitToTarget.Normalize();
+		const FVector ToTarget = TargetTM.GetTranslation() - SourceTM.GetTranslation();
+		const FVector UnitToTarget = ToTarget.GetSafeNormal();
 		FVector Midpoint = SourceTM.GetTranslation() + 0.5f * ToTarget + 0.5f * UnitToTarget * ArrowHeadHeight;
 
 		FVector YAxis;
 		FVector ZAxis;
 		UnitToTarget.FindBestAxisVectors(YAxis, ZAxis);
-		FMatrix ArrowMatrix(UnitToTarget, YAxis, ZAxis, Midpoint);
+		const FMatrix ArrowMatrix(UnitToTarget, YAxis, ZAxis, Midpoint);
 
 		DrawConnectedArrow(PDI, ArrowMatrix, FLinearColor(0.0f, 1.0f, 0.0), ArrowHeadHeight, ArrowHeadWidth, SDPG_Foreground);
 
@@ -83,4 +84,52 @@ void UAnimGraphNode_BoneDrivenController::Draw(FPrimitiveDrawInterface* PDI, USk
 	}
 }
 
+void UAnimGraphNode_BoneDrivenController::ValidateAnimNodeDuringCompilation(USkeleton* ForSkeleton, FCompilerResultsLog& MessageLog)
+{
+	if (ForSkeleton->GetReferenceSkeleton().FindBoneIndex(Node.SourceBone.BoneName) == INDEX_NONE)
+	{
+		MessageLog.Warning(*LOCTEXT("NoSourceBone", "@@ - You must pick a source bone as the Driver joint").ToString(), this);
+	}
+
+	if (Node.SourceComponent == EComponentType::None)
+	{
+		MessageLog.Warning(*LOCTEXT("NoSourceComponent", "@@ - You must pick a source component on the Driver joint").ToString(), this);
+	}
+	
+	if (ForSkeleton->GetReferenceSkeleton().FindBoneIndex(Node.TargetBone.BoneName) == INDEX_NONE)
+	{
+		MessageLog.Warning(*LOCTEXT("NoTargetBone", "@@ - You must pick a target bone as the Driven joint").ToString(), this);
+	}
+
+	if (Node.TargetComponent == EComponentType::None)
+	{
+		MessageLog.Warning(*LOCTEXT("NoTargetComponent", "@@ - You must pick a target component on the Driven joint").ToString(), this);
+	}
+
+	Super::ValidateAnimNodeDuringCompilation(ForSkeleton, MessageLog);
+}
+
+void UAnimGraphNode_BoneDrivenController::CustomizeDetails(IDetailLayoutBuilder& DetailBuilder)
+{
+	DetailBuilder.EditCategory(TEXT("Source (Driver)"));
+	IDetailCategoryBuilder& MappingCategory = DetailBuilder.EditCategory(TEXT("Mapping"));
+	DetailBuilder.EditCategory(TEXT("Destination (Driven)"));
+
+	TSharedRef<IPropertyHandle> NodeHandle = DetailBuilder.GetProperty(GET_MEMBER_NAME_CHECKED(UAnimGraphNode_BoneDrivenController, Node), GetClass());
+
+	MappingCategory.AddProperty(NodeHandle->GetChildHandle(GET_MEMBER_NAME_CHECKED(FAnimNode_BoneDrivenController, DrivingCurve)));
+
+	TAttribute<EVisibility> MappingVisibility = TAttribute<EVisibility>::Create(TAttribute<EVisibility>::FGetter::CreateUObject(this, &UAnimGraphNode_BoneDrivenController::AreNonCurveMappingValuesVisible));
+	MappingCategory.AddProperty(NodeHandle->GetChildHandle(GET_MEMBER_NAME_CHECKED(FAnimNode_BoneDrivenController, Multiplier))).Visibility(MappingVisibility);
+	MappingCategory.AddProperty(NodeHandle->GetChildHandle(GET_MEMBER_NAME_CHECKED(FAnimNode_BoneDrivenController, bUseRange))).Visibility(MappingVisibility);
+	MappingCategory.AddProperty(NodeHandle->GetChildHandle(GET_MEMBER_NAME_CHECKED(FAnimNode_BoneDrivenController, RangeMin))).Visibility(MappingVisibility);
+	MappingCategory.AddProperty(NodeHandle->GetChildHandle(GET_MEMBER_NAME_CHECKED(FAnimNode_BoneDrivenController, RangeMax))).Visibility(MappingVisibility);
+}
+
+EVisibility UAnimGraphNode_BoneDrivenController::AreNonCurveMappingValuesVisible() const
+{
+	return (Node.DrivingCurve != nullptr) ? EVisibility::Collapsed : EVisibility::Visible;
+}
+
+
 #undef LOCTEXT_NAMESPACE
diff --git a/Engine/Source/Runtime/AnimGraphRuntime/Private/BoneControllers/AnimNode_BoneDrivenController.cpp b/Engine/Source/Runtime/AnimGraphRuntime/Private/BoneControllers/AnimNode_BoneDrivenController.cpp
--- a/Engine/Source/Runtime/AnimGraphRuntime/Private/BoneControllers/AnimNode_BoneDrivenController.cpp
+++ b/Engine/Source/Runtime/AnimGraphRuntime/Private/BoneControllers/AnimNode_BoneDrivenController.cpp
@@ -8,11 +8,11 @@
 
 FAnimNode_BoneDrivenController::FAnimNode_BoneDrivenController()
 	: SourceComponent(EComponentType::None)
-	, TargetComponent(EComponentType::None)
 	, Multiplier(1.0f)
 	, bUseRange(false)
 	, RangeMin(-1.0f)
 	, RangeMax(1.0f)
+	, TargetComponent(EComponentType::None)
 {
 }
 
@@ -32,7 +32,7 @@ void FAnimNode_BoneDrivenController::EvaluateBoneTransforms(USkeletalMeshCompone
 	check(OutBoneTransforms.Num() == 0);
 	
 	// Early out if we're not driving from or to anything
-	if(SourceComponent == EComponentType::None || TargetComponent == EComponentType::None || Multiplier == 0.0f)
+	if ((SourceComponent == EComponentType::None) || (TargetComponent == EComponentType::None))
 	{
 		return;
 	}
@@ -40,43 +40,57 @@ void FAnimNode_BoneDrivenController::EvaluateBoneTransforms(USkeletalMeshCompone
 	// Get the Local space transform and the ref pose transform to see how the transform for the source bone has changed
 	const FBoneContainer& BoneContainer = MeshBases.GetPose().GetBoneContainer();
 	const FTransform& SourceOrigRef = BoneContainer.GetRefPoseArray()[SourceBone.BoneIndex];
-	FTransform SourceCurr = MeshBases.GetLocalSpaceTransform(SourceBone.GetCompactPoseIndex(BoneContainer));
-
-	// Difference from the ref pose
-	FVector RotationDiff = (SourceCurr.GetRotation() * SourceOrigRef.GetRotation().Inverse()).Euler();
-	FVector TranslationDiff = SourceCurr.GetLocation() - SourceOrigRef.GetLocation();
-	FVector CurrentScale = SourceCurr.GetScale3D();
-	FVector RefScale = SourceOrigRef.GetScale3D();
-	float ScaleDiff = FMath::Max3(CurrentScale[0], CurrentScale[1], CurrentScale[2]) - FMath::Max3(RefScale[0], RefScale[1], RefScale[2]);
-	
-	// Difference to apply after processing
-	FVector NewRotDiff(FVector::ZeroVector);
-	FVector NewTransDiff(FVector::ZeroVector);
-	float NewScaleDiff = 0.0f;
-
-	float* SourcePtr = nullptr;
-	float* DestPtr = nullptr;
+	const FTransform SourceCurr = MeshBases.GetLocalSpaceTransform(SourceBone.GetCompactPoseIndex(BoneContainer));
 
 	// Resolve source value
-	if(SourceComponent < EComponentType::RotationX)
+	float SourceValue = 0.0f;
+	if (SourceComponent < EComponentType::RotationX)
 	{
-		SourcePtr = &TranslationDiff[(int32)(SourceComponent - EComponentType::TranslationX)];
+		const FVector TranslationDiff = SourceCurr.GetLocation() - SourceOrigRef.GetLocation();
+		SourceValue = TranslationDiff[(int32)(SourceComponent - EComponentType::TranslationX)];
 	}
-	else if(SourceComponent < EComponentType::Scale)
+	else if (SourceComponent < EComponentType::Scale)
 	{
-		SourcePtr = &RotationDiff[(int32)(SourceComponent - EComponentType::RotationX)];
+		const FVector RotationDiff = (SourceCurr.GetRotation() * SourceOrigRef.GetRotation().Inverse()).Euler();
+		SourceValue = RotationDiff[(int32)(SourceComponent - EComponentType::RotationX)];
 	}
 	else
 	{
-		SourcePtr = &ScaleDiff;
+		const FVector CurrentScale = SourceCurr.GetScale3D();
+		const FVector RefScale = SourceOrigRef.GetScale3D();
+		const float ScaleDiff = FMath::Max3(CurrentScale[0], CurrentScale[1], CurrentScale[2]) - FMath::Max3(RefScale[0], RefScale[1], RefScale[2]);
+		SourceValue = ScaleDiff;
 	}
 	
+	// Determine the resulting value
+	float FinalDriverValue = SourceValue;
+	if (DrivingCurve != nullptr)
+	{
+		// Remap thru the curve if set
+		FinalDriverValue = DrivingCurve->GetFloatValue(FinalDriverValue);
+	}
+	else
+	{
+		// Apply the fixed function remapping/clamping
+		FinalDriverValue *= Multiplier;
+		if (bUseRange)
+		{
+			FinalDriverValue = FMath::Clamp(FinalDriverValue, RangeMin, RangeMax);
+		}
+	}
+
+	// Difference to apply after processing
+	FVector NewRotDiff(FVector::ZeroVector);
+	FVector NewTransDiff(FVector::ZeroVector);
+	float NewScaleDiff = 0.0f;
+
 	// Resolve target value
-	if(TargetComponent < EComponentType::RotationX)
+	float* DestPtr = nullptr;
+	if (TargetComponent < EComponentType::RotationX)
 	{
 		DestPtr = &NewTransDiff[(int32)(TargetComponent - EComponentType::TranslationX)];
 	}
-	else if(TargetComponent < EComponentType::Scale)
+	else if (TargetComponent < EComponentType::Scale)
 	{
 		DestPtr = &NewRotDiff[(int32)(TargetComponent - EComponentType::RotationX)];
 	}
@@ -85,31 +99,24 @@ void FAnimNode_BoneDrivenController::EvaluateBoneTransforms(USkeletalMeshCompone
 		DestPtr = &NewScaleDiff;
 	}
 
-	// Set the value; clamped to the given range if necessary.
-	if(bUseRange)
-	{
-		*DestPtr = FMath::Clamp(*SourcePtr * Multiplier, RangeMin, RangeMax);
-	}
-	else
-	{
-		*DestPtr = *SourcePtr * Multiplier;
-	}
+	// Set the value
+	*DestPtr = FinalDriverValue;
 
 	// Build final transform difference
-	FTransform FinalDiff(FQuat::MakeFromEuler(NewRotDiff), NewTransDiff, FVector(NewScaleDiff));
+	const FTransform FinalDiff(FQuat::MakeFromEuler(NewRotDiff), NewTransDiff, FVector(NewScaleDiff));
 
-	FCompactPoseBoneIndex TargetBoneIndex = TargetBone.GetCompactPoseIndex(BoneContainer);
+	const FCompactPoseBoneIndex TargetBoneIndex = TargetBone.GetCompactPoseIndex(BoneContainer);
 
 	// Starting point for the new transform
 	FTransform NewLocal = MeshBases.GetLocalSpaceTransform(TargetBoneIndex);
 	NewLocal.AccumulateWithAdditiveScale3D(FinalDiff);
 
 	// If we have a parent, concatenate the transform, otherwise just take the new transform
-	FCompactPoseBoneIndex PIdx = MeshBases.GetPose().GetParentBoneIndex(TargetBoneIndex);
+	const FCompactPoseBoneIndex ParentIndex = MeshBases.GetPose().GetParentBoneIndex(TargetBoneIndex);
 
-	if(PIdx != INDEX_NONE)
+	if (ParentIndex != INDEX_NONE)
 	{
-		const FTransform& ParentTM = MeshBases.GetComponentSpaceTransform(PIdx);
+		const FTransform& ParentTM = MeshBases.GetComponentSpaceTransform(ParentIndex);
 		
 		OutBoneTransforms.Add(FBoneTransform(TargetBoneIndex, NewLocal * ParentTM));
 	}
@@ -117,7 +124,6 @@ void FAnimNode_BoneDrivenController::EvaluateBoneTransforms(USkeletalMeshCompone
 	{
 		OutBoneTransforms.Add(FBoneTransform(TargetBoneIndex, NewLocal));
 	}
-
 }
 
 bool FAnimNode_BoneDrivenController::IsValidToEvaluate(const USkeleton* Skeleton, const FBoneContainer& RequiredBones)
diff --git a/Engine/Source/Runtime/AnimGraphRuntime/Public/BoneControllers/AnimNode_BoneDrivenController.h b/Engine/Source/Runtime/AnimGraphRuntime/Public/BoneControllers/AnimNode_BoneDrivenController.h
--- a/Engine/Source/Runtime/AnimGraphRuntime/Public/BoneControllers/AnimNode_BoneDrivenController.h
+++ b/Engine/Source/Runtime/AnimGraphRuntime/Public/BoneControllers/AnimNode_BoneDrivenController.h
@@ -23,57 +23,65 @@ namespace EComponentType
 	};
 }
 
+/**
+ * This is the runtime version of a bone driven controller, which maps part of the state from one bone to another (e.g., 2 * source.x -> target.z)
+ */
 USTRUCT()
 struct ANIMGRAPHRUNTIME_API FAnimNode_BoneDrivenController : public FAnimNode_SkeletalControlBase
 {
 	GENERATED_USTRUCT_BODY()
 
 	// Bone to use as controller input
-	UPROPERTY(EditAnywhere, BlueprintReadWrite, Category = Bones)
+	UPROPERTY(EditAnywhere, Category="Source (driver)")
 	FBoneReference SourceBone;
 
-	// Bone to drive using controller input
-	UPROPERTY(EditAnywhere, BlueprintReadWrite, Category = Bones)
-	FBoneReference TargetBone;
-
 	// Transform component to use as input
-	UPROPERTY(EditAnywhere, BlueprintReadWrite, Category = Controller)
+	UPROPERTY(EditAnywhere, Category="Source (driver)")
 	TEnumAsByte<EComponentType::Type> SourceComponent;
 
-	// Transform component to write to
-	UPROPERTY(EditAnywhere, BlueprintReadWrite, Category = Controller)
-	TEnumAsByte<EComponentType::Type> TargetComponent;
+	/** Curve used to map from the source attribute to the driven attributes if present (otherwise the Multiplier will be used) */
+	UPROPERTY(EditAnywhere, Category=Mapping)
+	UCurveFloat* DrivingCurve;
 
 	// Multiplier to apply to the input value
-	UPROPERTY(EditAnywhere, BlueprintReadWrite, Category=Controller)
+	UPROPERTY(EditAnywhere, Category=Mapping)
 	float Multiplier;
 
-	// Whether or not to use the range limits
-	UPROPERTY(EditAnywhere, BlueprintReadWrite, Category=Range)
+	// Whether or not to use the range limits on the multiplied value
+	UPROPERTY(EditAnywhere, Category=Mapping)
 	bool bUseRange;
 
-	// Min limit of the output value
-	UPROPERTY(EditAnywhere, BlueprintReadWrite, Category=Range)
+	// Minimum limit of the output value
+	UPROPERTY(EditAnywhere, Category=Mapping, meta=(EditCondition=bUseRange))
 	float RangeMin;
 
-	// Max limit of the output value
-	UPROPERTY(EditAnywhere, BlueprintReadWrite, Category=Range)
+	// Maximum limit of the output value
+	UPROPERTY(EditAnywhere, Category=Mapping, meta=(EditCondition=bUseRange))
 	float RangeMax;
 
+	// Bone to drive using controller input
+	UPROPERTY(EditAnywhere, Category="Destination (driven)")
+	FBoneReference TargetBone;
+
+	// Transform component to write to
+	UPROPERTY(EditAnywhere, Category="Destination (driven)")
+	TEnumAsByte<EComponentType::Type> TargetComponent;
+
+public:
 	FAnimNode_BoneDrivenController();
 
 	// FAnimNode_Base interface
 	virtual void GatherDebugData(FNodeDebugData& DebugData) override;
-	//////////////////////////////////////////////////////////////////////////
+	// End of FAnimNode_Base interface
 
 	// FAnimNode_SkeletalControlBase interface
 	virtual void EvaluateBoneTransforms(USkeletalMeshComponent* SkelComp, FCSPose<FCompactPose>& MeshBases, TArray<FBoneTransform>& OutBoneTransforms) override;
 	virtual bool IsValidToEvaluate(const USkeleton* Skeleton, const FBoneContainer& RequiredBones) override;
-	//////////////////////////////////////////////////////////////////////////
+	// End of FAnimNode_SkeletalControlBase interface
 
 protected:
 	
-	// FAnimNode_SkeletalControlBase interface
+	// FAnimNode_SkeletalControlBase protected interface
 	virtual void InitializeBoneReferences(const FBoneContainer& RequiredBones) override;
-	//////////////////////////////////////////////////////////////////////////
+	// End of FAnimNode_SkeletalControlBase protected interface
 };
\ No newline at end of file
