Animation: Adding missing file for custom version

diff --git a/Engine/Source/Runtime/AnimGraphRuntime/Private/AnimationCustomVersion.cpp b/Engine/Source/Runtime/AnimGraphRuntime/Private/AnimationCustomVersion.cpp
new file mode 100644
--- /dev/null
+++ b/Engine/Source/Runtime/AnimGraphRuntime/Private/AnimationCustomVersion.cpp
@@ -0,0 +1,9 @@
+// Copyright 1998-2015 Epic Games, Inc. All Rights Reserved.
+
+#include "AnimGraphRuntimePrivatePCH.h"
+#include "AnimationCustomVersion.h"
+
+const FGuid FAnimationCustomVersion::GUID(0x11310AED, 0x2E554D61, 0xAF679AA3, 0xC5A1082C);
+
+// Register the custom version with core
+FCustomVersionRegistration GRegisterAnimationCustomVersion(FAnimationCustomVersion::GUID, FAnimationCustomVersion::LatestVersion, TEXT("AnimGraphVer"));
