Animation: Regenerated the GUID for the animation custom version to avoid a collision

diff --git a/Engine/Source/Runtime/AnimGraphRuntime/Private/AnimationCustomVersion.cpp b/Engine/Source/Runtime/AnimGraphRuntime/Private/AnimationCustomVersion.cpp
--- a/Engine/Source/Runtime/AnimGraphRuntime/Private/AnimationCustomVersion.cpp
+++ b/Engine/Source/Runtime/AnimGraphRuntime/Private/AnimationCustomVersion.cpp
@@ -3,7 +3,7 @@
 #include "AnimGraphRuntimePrivatePCH.h"
 #include "AnimationCustomVersion.h"
 
-const FGuid FAnimationCustomVersion::GUID(0x11310AED, 0x2E554D61, 0xAF679AA3, 0xC5A1082C);
+const FGuid FAnimationCustomVersion::GUID(0x2EB5FDBD, 0x01AC4D10, 0x8136F38F, 0x3393A5DA);
 
 // Register the custom version with core
 FCustomVersionRegistration GRegisterAnimationCustomVersion(FAnimationCustomVersion::GUID, FAnimationCustomVersion::LatestVersion, TEXT("AnimGraphVer"));
