Animation: Added support for showing the position relative to the reference skeleton on the observe bone node

diff --git a/Engine/Source/Runtime/AnimGraphRuntime/Private/BoneControllers/AnimNode_ObserveBone.cpp b/Engine/Source/Runtime/AnimGraphRuntime/Private/BoneControllers/AnimNode_ObserveBone.cpp
--- a/Engine/Source/Runtime/AnimGraphRuntime/Private/BoneControllers/AnimNode_ObserveBone.cpp
+++ b/Engine/Source/Runtime/AnimGraphRuntime/Private/BoneControllers/AnimNode_ObserveBone.cpp
@@ -9,6 +9,7 @@
 
 FAnimNode_ObserveBone::FAnimNode_ObserveBone()
 	: DisplaySpace(BCS_ComponentSpace)
+	, bRelativeToRefPose(false)
 	, Translation(FVector::ZeroVector)
 	, Rotation(FRotator::ZeroRotator)
 	, Scale(FVector(1.0f))
@@ -34,6 +35,14 @@ void FAnimNode_ObserveBone::EvaluateBoneTransforms(USkeletalMeshComponent* SkelC
 	// Convert to the specific display space if necessary
 	FAnimationRuntime::ConvertCSTransformToBoneSpace(SkelComp, MeshBases, BoneTM, BoneIndex, DisplaySpace);
 
+	// Convert to be relative to the ref pose if necessary
+	if (bRelativeToRefPose)
+	{
+		const FTransform& SourceOrigRef = BoneContainer.GetRefPoseArray()[BoneToObserve.BoneIndex];
+		BoneTM = BoneTM.GetRelativeTransform(SourceOrigRef);
+	}
+
+	// Cache off the values for display
 	Translation = BoneTM.GetTranslation();
 	Rotation = BoneTM.GetRotation().Rotator();
 	Scale = BoneTM.GetScale3D();
diff --git a/Engine/Source/Runtime/AnimGraphRuntime/Public/BoneControllers/AnimNode_ObserveBone.h b/Engine/Source/Runtime/AnimGraphRuntime/Public/BoneControllers/AnimNode_ObserveBone.h
--- a/Engine/Source/Runtime/AnimGraphRuntime/Public/BoneControllers/AnimNode_ObserveBone.h
+++ b/Engine/Source/Runtime/AnimGraphRuntime/Public/BoneControllers/AnimNode_ObserveBone.h
@@ -14,13 +14,17 @@ struct ANIMGRAPHRUNTIME_API FAnimNode_ObserveBone : public FAnimNode_SkeletalCon
 	GENERATED_USTRUCT_BODY()
 
 	/** Name of bone to observe. */
-	UPROPERTY(EditAnywhere, BlueprintReadWrite, Category=SkeletalControl) 
+	UPROPERTY(EditAnywhere, Category=SkeletalControl)
 	FBoneReference BoneToObserve;
 
 	/** Reference frame to display the bone transform in. */
-	UPROPERTY(EditAnywhere, BlueprintReadWrite, Category=SkeletalControl)
+	UPROPERTY(EditAnywhere, Category=SkeletalControl)
 	TEnumAsByte<EBoneControlSpace> DisplaySpace;
 
+	/** Show the difference from the reference pose? */
+	UPROPERTY(EditAnywhere, Category=SkeletalControl)
+	bool bRelativeToRefPose;
+
 	/** Translation of the bone being observed. */
 	UPROPERTY()
 	FVector Translation;
