#ANIM: Created [copy pose from mesh] node.

Allows users to get a pose from one mesh to the other pose.

diff --git a/Engine/Source/Editor/AnimGraph/Classes/AnimGraphNode_CopyPoseFromMesh.h b/Engine/Source/Editor/AnimGraph/Classes/AnimGraphNode_CopyPoseFromMesh.h
new file mode 100644
--- /dev/null
+++ b/Engine/Source/Editor/AnimGraph/Classes/AnimGraphNode_CopyPoseFromMesh.h
@@ -0,0 +1,22 @@
+// Copyright 1998-2015 Epic Games, Inc. All Rights Reserved.
+
+#pragma once
+#include "AnimGraphNode_SkeletalControlBase.h"
+#include "AnimNodes/AnimNode_CopyPoseFromMesh.h"
+#include "EdGraph/EdGraphNodeUtils.h" // for FNodeTitleTextTable
+#include "AnimGraphNode_CopyPoseFromMesh.generated.h"
+
+UCLASS(MinimalAPI)
+class UAnimGraphNode_CopyPoseFromMesh : public UAnimGraphNode_Base
+{
+	GENERATED_UCLASS_BODY()
+
+	UPROPERTY(EditAnywhere, Category=Settings)
+	FAnimNode_CopyPoseFromMesh Node;
+
+public:
+	// UEdGraphNode interface
+	virtual FText GetNodeTitle(ENodeTitleType::Type TitleType) const override;
+	virtual FText GetTooltipText() const override;
+	// End of UEdGraphNode interface
+};
diff --git a/Engine/Source/Editor/AnimGraph/Private/AnimGraphNode_CopyPoseFromMesh.cpp b/Engine/Source/Editor/AnimGraph/Private/AnimGraphNode_CopyPoseFromMesh.cpp
new file mode 100644
--- /dev/null
+++ b/Engine/Source/Editor/AnimGraph/Private/AnimGraphNode_CopyPoseFromMesh.cpp
@@ -0,0 +1,26 @@
+// Copyright 1998-2015 Epic Games, Inc. All Rights Reserved.
+
+#include "AnimGraphPrivatePCH.h"
+#include "AnimGraphNode_CopyPoseFromMesh.h"
+
+/////////////////////////////////////////////////////
+// UAnimGraphNode_CopyPoseFromMeshSkeletalControl
+
+#define LOCTEXT_NAMESPACE "A3Nodes"
+
+UAnimGraphNode_CopyPoseFromMesh::UAnimGraphNode_CopyPoseFromMesh(const FObjectInitializer& ObjectInitializer)
+	: Super(ObjectInitializer)
+{
+}
+
+FText UAnimGraphNode_CopyPoseFromMesh::GetTooltipText() const
+{
+	return LOCTEXT("AnimGraphNode_CopyPoseFromMesh_Tooltip", "The Copy Pose From Mesh node copies the pose data from another component to this. Only works when name matches.");
+}
+
+FText UAnimGraphNode_CopyPoseFromMesh::GetNodeTitle(ENodeTitleType::Type TitleType) const
+{
+	return LOCTEXT("CopyPoseFromMesh", "Copy Pose From Mesh");
+}
+
+#undef LOCTEXT_NAMESPACE
diff --git a/Engine/Source/Runtime/AnimGraphRuntime/Private/AnimNodes/AnimNode_CopyPoseFromMesh.cpp b/Engine/Source/Runtime/AnimGraphRuntime/Private/AnimNodes/AnimNode_CopyPoseFromMesh.cpp
new file mode 100644
--- /dev/null
+++ b/Engine/Source/Runtime/AnimGraphRuntime/Private/AnimNodes/AnimNode_CopyPoseFromMesh.cpp
@@ -0,0 +1,102 @@
+// Copyright 1998-2015 Epic Games, Inc. All Rights Reserved.
+
+#include "AnimGraphRuntimePrivatePCH.h"
+#include "AnimNodes/AnimNode_CopyPoseFromMesh.h"
+
+/////////////////////////////////////////////////////
+// FAnimNode_CopyPoseFromMesh
+
+FAnimNode_CopyPoseFromMesh::FAnimNode_CopyPoseFromMesh()
+{
+}
+
+void FAnimNode_CopyPoseFromMesh::Initialize(const FAnimationInitializeContext& Context)
+{
+	ReinitializeMeshComponent(Context.AnimInstance);
+}
+
+void FAnimNode_CopyPoseFromMesh::CacheBones(const FAnimationCacheBonesContext& Context)
+{
+}
+
+void FAnimNode_CopyPoseFromMesh::Update(const FAnimationUpdateContext& Context)
+{
+	EvaluateGraphExposedInputs.Execute(Context);
+
+	if (CurrentlyUsedMeshComponent.IsValid() && CurrentlyUsedMeshComponent.Get() != SourceMeshComponent)
+	{
+		ReinitializeMeshComponent(Context.AnimInstance);
+	}
+	else if (!CurrentlyUsedMeshComponent.IsValid() && SourceMeshComponent)
+	{
+		ReinitializeMeshComponent(Context.AnimInstance);
+	}
+}
+
+void FAnimNode_CopyPoseFromMesh::Evaluate(FPoseContext& Output)
+{
+	FCompactPose& OutPose = Output.Pose;
+	OutPose.ResetToRefPose();
+
+	if (CurrentlyUsedMeshComponent.IsValid() && CurrentlyUsedMeshComponent.Get()->SkeletalMesh)
+	{
+		const FBoneContainer& RequiredBones = OutPose.GetBoneContainer();
+		for(FCompactPoseBoneIndex PoseBoneIndex : OutPose.ForEachBoneIndex())
+		{
+			const int32& SkeletonBoneIndex = RequiredBones.GetSkeletonIndex(PoseBoneIndex);
+			const int32* Value = BoneMapToSource.Find(SkeletonBoneIndex);
+			if(Value && *Value!=INDEX_NONE)
+			{
+				const int32 SourceBoneIndex = *Value;
+				const int32 ParentIndex = CurrentlyUsedMeshComponent.Get()->SkeletalMesh->RefSkeleton.GetParentIndex(SourceBoneIndex);
+				const FCompactPoseBoneIndex MyParentIndex = RequiredBones.GetParentBoneIndex(PoseBoneIndex);
+				// only apply if I also have parent, otherwise, it should apply the space bases
+				if (ParentIndex != INDEX_NONE && MyParentIndex != INDEX_NONE)
+				{
+					const FTransform& ParentTransform = CurrentlyUsedMeshComponent.Get()->GetSpaceBases()[ParentIndex];
+					const FTransform& ChildTransform = CurrentlyUsedMeshComponent.Get()->GetSpaceBases()[SourceBoneIndex]; 
+					OutPose[PoseBoneIndex] = ChildTransform.GetRelativeTransform(ParentTransform);
+				}
+				else
+				{
+					OutPose[PoseBoneIndex] = CurrentlyUsedMeshComponent.Get()->GetSpaceBases()[SourceBoneIndex]; 
+				}
+			}
+		}
+	}
+}
+
+void FAnimNode_CopyPoseFromMesh::GatherDebugData(FNodeDebugData& DebugData)
+{
+}
+
+void FAnimNode_CopyPoseFromMesh::ReinitializeMeshComponent(UAnimInstance* AnimInstance)
+{
+	CurrentlyUsedMeshComponent = SourceMeshComponent;
+	BoneMapToSource.Reset();
+	if (SourceMeshComponent && SourceMeshComponent->SkeletalMesh && !SourceMeshComponent->IsPendingKill())
+	{
+		USkeletalMeshComponent* TargetMeshComponent = AnimInstance->GetSkelMeshComponent();
+		if (TargetMeshComponent)
+		{
+			USkeletalMesh* SourceSkelMesh = SourceMeshComponent->SkeletalMesh;
+			USkeletalMesh* TargetSkelMesh = TargetMeshComponent->SkeletalMesh;
+
+			if (SourceSkelMesh == TargetSkelMesh)
+			{
+				for(int32 ComponentSpaceBoneId = 0; ComponentSpaceBoneId < SourceSkelMesh->RefSkeleton.GetNum(); ++ComponentSpaceBoneId)
+				{
+					BoneMapToSource.Add(ComponentSpaceBoneId, ComponentSpaceBoneId);
+				}
+			}
+			else
+			{
+				for (int32 ComponentSpaceBoneId=0; ComponentSpaceBoneId<TargetSkelMesh->RefSkeleton.GetNum(); ++ComponentSpaceBoneId)
+				{
+					FName BoneName = TargetSkelMesh->RefSkeleton.GetBoneName(ComponentSpaceBoneId);
+					BoneMapToSource.Add(ComponentSpaceBoneId, SourceSkelMesh->RefSkeleton.FindBoneIndex(BoneName));
+				}
+			}
+		}
+	}
+}
\ No newline at end of file
diff --git a/Engine/Source/Runtime/AnimGraphRuntime/Public/AnimNodes/AnimNode_CopyPoseFromMesh.h b/Engine/Source/Runtime/AnimGraphRuntime/Public/AnimNodes/AnimNode_CopyPoseFromMesh.h
new file mode 100644
--- /dev/null
+++ b/Engine/Source/Runtime/AnimGraphRuntime/Public/AnimNodes/AnimNode_CopyPoseFromMesh.h
@@ -0,0 +1,38 @@
+// Copyright 1998-2015 Epic Games, Inc. All Rights Reserved.
+
+#pragma once
+
+#include "AnimNode_SkeletalControlBase.h"
+#include "AnimNode_CopyPoseFromMesh.generated.h"
+
+/**
+ *	Simple controller to copy a bone's transform to another one.
+ */
+
+USTRUCT()
+struct ANIMGRAPHRUNTIME_API FAnimNode_CopyPoseFromMesh : public FAnimNode_Base
+{
+	GENERATED_USTRUCT_BODY()
+
+	UPROPERTY(EditAnywhere, BlueprintReadWrite, Category=Copy, meta=(PinShownByDefault))
+	USkeletalMeshComponent* SourceMeshComponent;
+
+	FAnimNode_CopyPoseFromMesh();
+
+	// FAnimNode_Base interface
+	virtual void Initialize(const FAnimationInitializeContext& Context) override;
+	virtual void CacheBones(const FAnimationCacheBonesContext& Context) override;
+	virtual void Update(const FAnimationUpdateContext& Context) override;
+	virtual void Evaluate(FPoseContext& Output) override;
+	virtual void GatherDebugData(FNodeDebugData& DebugData) override;
+	// End of FAnimNode_Base interface
+
+private:
+	TWeakObjectPtr<USkeletalMeshComponent> CurrentlyUsedMeshComponent;
+	// cache of target space bases to source space bases
+	TMap<int32, int32> BoneMapToSource;
+	int32 CachedTargetLODIndex;
+
+	// reinitialize mesh component 
+	void ReinitializeMeshComponent(UAnimInstance* AnimInstance);
+};
