#ANIM: Moving available animnodes to AnimGraphRuntime module

diff --git a/Engine/Config/BaseEngine.ini b/Engine/Config/BaseEngine.ini
--- a/Engine/Config/BaseEngine.ini
+++ b/Engine/Config/BaseEngine.ini
@@ -382,6 +382,24 @@ bUseStreamingPause=false
 +ActiveClassRedirects=(OldClassName="AnimNode_TwoBoneIK",NewClassName="/Script/AnimGraphRuntime.AnimNode_TwoBoneIK")
 +ActiveClassRedirects=(OldClassName="AnimNode_WheelHandler",NewClassName="/Script/AnimGraphRuntime.AnimNode_WheelHandler")
 
+; Moving more animation nodes form Engine to AnimGraphRuntime
++ActiveClassRedirects=(OldClassName="AnimNode_ApplyAdditive",NewClassName="/Script/AnimGraphRuntime.AnimNode_ApplyAdditive")
++ActiveClassRedirects=(OldClassName="AnimNode_BlendListBase",NewClassName="/Script/AnimGraphRuntime.AnimNode_BlendListBase")
++ActiveClassRedirects=(OldClassName="AnimNode_BlendListByBool",NewClassName="/Script/AnimGraphRuntime.AnimNode_BlendListByBool")
++ActiveClassRedirects=(OldClassName="AnimNode_BlendListByEnum",NewClassName="/Script/AnimGraphRuntime.AnimNode_BlendListByEnum")
++ActiveClassRedirects=(OldClassName="AnimNode_BlendListByInt",NewClassName="/Script/AnimGraphRuntime.AnimNode_BlendListByInt")
++ActiveClassRedirects=(OldClassName="AnimNode_BlendSpaceEvaluator",NewClassName="/Script/AnimGraphRuntime.AnimNode_BlendSpaceEvaluator")
++ActiveClassRedirects=(OldClassName="AnimNode_BlendSpacePlayer",NewClassName="/Script/AnimGraphRuntime.AnimNode_BlendSpacePlayer")
++ActiveClassRedirects=(OldClassName="AnimNode_LayeredBoneBlend",NewClassName="/Script/AnimGraphRuntime.AnimNode_LayeredBoneBlend")
++ActiveClassRedirects=(OldClassName="AnimNode_RefPose",NewClassName="/Script/AnimGraphRuntime.AnimNode_RefPose")
++ActiveClassRedirects=(OldClassName="AnimNode_Root",NewClassName="/Script/AnimGraphRuntime.AnimNode_Root")
++ActiveClassRedirects=(OldClassName="AnimNode_RotateRootBone",NewClassName="/Script/AnimGraphRuntime.AnimNode_RotateRootBone")
++ActiveClassRedirects=(OldClassName="AnimNode_RotationOffsetBlendSpace",NewClassName="/Script/AnimGraphRuntime.AnimNode_RotationOffsetBlendSpace")
++ActiveClassRedirects=(OldClassName="AnimNode_SaveCachedPose",NewClassName="/Script/AnimGraphRuntime.AnimNode_SaveCachedPose")
++ActiveClassRedirects=(OldClassName="AnimNode_SequenceEvaluator",NewClassName="/Script/AnimGraphRuntime.AnimNode_SequenceEvaluator")
++ActiveClassRedirects=(OldClassName="AnimNode_Slot",NewClassName="/Script/AnimGraphRuntime.AnimNode_Slot")
++ActiveClassRedirects=(OldClassName="AnimNode_TwoWayBlend",NewClassName="/Script/AnimGraphRuntime.AnimNode_TwoWayBlend")
+
 ; Renamed all user-exposed locations for rotator fields to display as XYZ (RPY) instead of YPR (4.9 era)
 +K2FieldRedirects=(OldFieldName="KismetMathLibrary.BreakRot", NewFieldName="KismetMathLibrary.BreakRotator")
 +K2FieldRedirects=(OldFieldName="KismetMathLibrary.MakeRot", NewFieldName="KismetMathLibrary.MakeRotator")
diff --git a/Engine/Source/Editor/AnimGraph/Classes/AnimGraphNode_ApplyAdditive.h b/Engine/Source/Editor/AnimGraph/Classes/AnimGraphNode_ApplyAdditive.h
--- a/Engine/Source/Editor/AnimGraph/Classes/AnimGraphNode_ApplyAdditive.h
+++ b/Engine/Source/Editor/AnimGraph/Classes/AnimGraphNode_ApplyAdditive.h
@@ -2,7 +2,7 @@
 
 #pragma once
 #include "AnimGraphNode_Base.h"
-#include "Animation/AnimNode_ApplyAdditive.h"
+#include "AnimNodes/AnimNode_ApplyAdditive.h"
 #include "AnimGraphNode_ApplyAdditive.generated.h"
 
 UCLASS(MinimalAPI)
diff --git a/Engine/Source/Editor/AnimGraph/Classes/AnimGraphNode_BlendListByBool.h b/Engine/Source/Editor/AnimGraph/Classes/AnimGraphNode_BlendListByBool.h
--- a/Engine/Source/Editor/AnimGraph/Classes/AnimGraphNode_BlendListByBool.h
+++ b/Engine/Source/Editor/AnimGraph/Classes/AnimGraphNode_BlendListByBool.h
@@ -2,7 +2,7 @@
 
 #pragma once
 #include "AnimGraphNode_BlendListBase.h"
-#include "Animation/AnimNode_BlendListByBool.h"
+#include "AnimNodes/AnimNode_BlendListByBool.h"
 #include "AnimGraphNode_BlendListByBool.generated.h"
 
 UCLASS(MinimalAPI)
diff --git a/Engine/Source/Editor/AnimGraph/Classes/AnimGraphNode_BlendListByEnum.h b/Engine/Source/Editor/AnimGraph/Classes/AnimGraphNode_BlendListByEnum.h
--- a/Engine/Source/Editor/AnimGraph/Classes/AnimGraphNode_BlendListByEnum.h
+++ b/Engine/Source/Editor/AnimGraph/Classes/AnimGraphNode_BlendListByEnum.h
@@ -2,7 +2,7 @@
 
 #pragma once
 #include "AnimGraphNode_BlendListBase.h"
-#include "Animation/AnimNode_BlendListByEnum.h"
+#include "AnimNodes/AnimNode_BlendListByEnum.h"
 #include "EdGraph/EdGraphNodeUtils.h" // for FNodeTextCache
 #include "AnimGraphNode_BlendListByEnum.generated.h"
 
diff --git a/Engine/Source/Editor/AnimGraph/Classes/AnimGraphNode_BlendListByInt.h b/Engine/Source/Editor/AnimGraph/Classes/AnimGraphNode_BlendListByInt.h
--- a/Engine/Source/Editor/AnimGraph/Classes/AnimGraphNode_BlendListByInt.h
+++ b/Engine/Source/Editor/AnimGraph/Classes/AnimGraphNode_BlendListByInt.h
@@ -2,7 +2,7 @@
 
 #pragma once
 #include "AnimGraphNode_BlendListBase.h"
-#include "Animation/AnimNode_BlendListByInt.h"
+#include "AnimNodes/AnimNode_BlendListByInt.h"
 #include "AnimGraphNode_BlendListByInt.generated.h"
 
 UCLASS(MinimalAPI)
diff --git a/Engine/Source/Editor/AnimGraph/Classes/AnimGraphNode_BlendSpaceEvaluator.h b/Engine/Source/Editor/AnimGraph/Classes/AnimGraphNode_BlendSpaceEvaluator.h
--- a/Engine/Source/Editor/AnimGraph/Classes/AnimGraphNode_BlendSpaceEvaluator.h
+++ b/Engine/Source/Editor/AnimGraph/Classes/AnimGraphNode_BlendSpaceEvaluator.h
@@ -2,7 +2,7 @@
 
 #pragma once
 #include "AnimGraphNode_BlendSpaceBase.h"
-#include "Animation/AnimNode_BlendSpaceEvaluator.h"
+#include "AnimNodes/AnimNode_BlendSpaceEvaluator.h"
 #include "EdGraph/EdGraphNodeUtils.h" // for FNodeTitleTextTable
 #include "AnimGraphNode_BlendSpaceEvaluator.generated.h"
 
diff --git a/Engine/Source/Editor/AnimGraph/Classes/AnimGraphNode_BlendSpacePlayer.h b/Engine/Source/Editor/AnimGraph/Classes/AnimGraphNode_BlendSpacePlayer.h
--- a/Engine/Source/Editor/AnimGraph/Classes/AnimGraphNode_BlendSpacePlayer.h
+++ b/Engine/Source/Editor/AnimGraph/Classes/AnimGraphNode_BlendSpacePlayer.h
@@ -2,7 +2,7 @@
 
 #pragma once
 #include "AnimGraphNode_BlendSpaceBase.h"
-#include "Animation/AnimNode_BlendSpacePlayer.h"
+#include "AnimNodes/AnimNode_BlendSpacePlayer.h"
 #include "EdGraph/EdGraphNodeUtils.h" // for FNodeTitleTextTable
 #include "AnimGraphNode_BlendSpacePlayer.generated.h"
 
diff --git a/Engine/Source/Editor/AnimGraph/Classes/AnimGraphNode_LayeredBoneBlend.h b/Engine/Source/Editor/AnimGraph/Classes/AnimGraphNode_LayeredBoneBlend.h
--- a/Engine/Source/Editor/AnimGraph/Classes/AnimGraphNode_LayeredBoneBlend.h
+++ b/Engine/Source/Editor/AnimGraph/Classes/AnimGraphNode_LayeredBoneBlend.h
@@ -2,7 +2,7 @@
 
 #pragma once
 #include "AnimGraphNode_BlendListBase.h"
-#include "Animation/AnimNode_LayeredBoneBlend.h"
+#include "AnimNodes/AnimNode_LayeredBoneBlend.h"
 #include "AnimGraphNode_LayeredBoneBlend.generated.h"
 
 UCLASS(MinimalAPI)
diff --git a/Engine/Source/Editor/AnimGraph/Classes/AnimGraphNode_MeshRefPose.h b/Engine/Source/Editor/AnimGraph/Classes/AnimGraphNode_MeshRefPose.h
--- a/Engine/Source/Editor/AnimGraph/Classes/AnimGraphNode_MeshRefPose.h
+++ b/Engine/Source/Editor/AnimGraph/Classes/AnimGraphNode_MeshRefPose.h
@@ -2,7 +2,7 @@
 
 #pragma once
 #include "AnimGraphNode_Base.h"
-#include "Animation/AnimNode_RefPose.h"
+#include "AnimNodes/AnimNode_RefPose.h"
 #include "AnimGraphNode_MeshRefPose.generated.h"
 
 UCLASS(MinimalAPI)
diff --git a/Engine/Source/Editor/AnimGraph/Classes/AnimGraphNode_RefPoseBase.h b/Engine/Source/Editor/AnimGraph/Classes/AnimGraphNode_RefPoseBase.h
--- a/Engine/Source/Editor/AnimGraph/Classes/AnimGraphNode_RefPoseBase.h
+++ b/Engine/Source/Editor/AnimGraph/Classes/AnimGraphNode_RefPoseBase.h
@@ -2,7 +2,7 @@
 
 #pragma once
 #include "AnimGraphNode_Base.h"
-#include "Animation/AnimNode_RefPose.h"
+#include "AnimNodes/AnimNode_RefPose.h"
 #include "AnimGraphNode_RefPoseBase.generated.h"
 
 UCLASS(MinimalAPI, abstract)
diff --git a/Engine/Source/Editor/AnimGraph/Classes/AnimGraphNode_Root.h b/Engine/Source/Editor/AnimGraph/Classes/AnimGraphNode_Root.h
--- a/Engine/Source/Editor/AnimGraph/Classes/AnimGraphNode_Root.h
+++ b/Engine/Source/Editor/AnimGraph/Classes/AnimGraphNode_Root.h
@@ -2,7 +2,7 @@
 
 #pragma once
 #include "AnimGraphNode_Base.h"
-#include "Animation/AnimNode_Root.h"
+#include "AnimNodes/AnimNode_Root.h"
 #include "AnimGraphNode_Root.generated.h"
 
 UCLASS(MinimalAPI)
diff --git a/Engine/Source/Editor/AnimGraph/Classes/AnimGraphNode_RotateRootBone.h b/Engine/Source/Editor/AnimGraph/Classes/AnimGraphNode_RotateRootBone.h
--- a/Engine/Source/Editor/AnimGraph/Classes/AnimGraphNode_RotateRootBone.h
+++ b/Engine/Source/Editor/AnimGraph/Classes/AnimGraphNode_RotateRootBone.h
@@ -3,7 +3,7 @@
 
 #pragma once
 #include "AnimGraphNode_Base.h"
-#include "Animation/AnimNode_RotateRootBone.h"
+#include "AnimNodes/AnimNode_RotateRootBone.h"
 #include "AnimGraphNode_RotateRootBone.generated.h"
 
 UCLASS(MinimalAPI)
diff --git a/Engine/Source/Editor/AnimGraph/Classes/AnimGraphNode_RotationOffsetBlendSpace.h b/Engine/Source/Editor/AnimGraph/Classes/AnimGraphNode_RotationOffsetBlendSpace.h
--- a/Engine/Source/Editor/AnimGraph/Classes/AnimGraphNode_RotationOffsetBlendSpace.h
+++ b/Engine/Source/Editor/AnimGraph/Classes/AnimGraphNode_RotationOffsetBlendSpace.h
@@ -2,7 +2,7 @@
 
 #pragma once
 #include "AnimGraphNode_BlendSpaceBase.h"
-#include "Animation/AnimNode_RotationOffsetBlendSpace.h"
+#include "AnimNodes/AnimNode_RotationOffsetBlendSpace.h"
 #include "EdGraph/EdGraphNodeUtils.h" // for FNodeTitleTextTable
 #include "AnimGraphNode_RotationOffsetBlendSpace.generated.h"
 
diff --git a/Engine/Source/Editor/AnimGraph/Classes/AnimGraphNode_SaveCachedPose.h b/Engine/Source/Editor/AnimGraph/Classes/AnimGraphNode_SaveCachedPose.h
--- a/Engine/Source/Editor/AnimGraph/Classes/AnimGraphNode_SaveCachedPose.h
+++ b/Engine/Source/Editor/AnimGraph/Classes/AnimGraphNode_SaveCachedPose.h
@@ -2,7 +2,7 @@
 
 #pragma once
 #include "AnimGraphNode_Base.h"
-#include "Animation/AnimNode_SaveCachedPose.h"
+#include "AnimNodes/AnimNode_SaveCachedPose.h"
 #include "EdGraph/EdGraphNodeUtils.h" // for FNodeTextCache
 #include "AnimGraphNode_SaveCachedPose.generated.h"
 
diff --git a/Engine/Source/Editor/AnimGraph/Classes/AnimGraphNode_SequenceEvaluator.h b/Engine/Source/Editor/AnimGraph/Classes/AnimGraphNode_SequenceEvaluator.h
--- a/Engine/Source/Editor/AnimGraph/Classes/AnimGraphNode_SequenceEvaluator.h
+++ b/Engine/Source/Editor/AnimGraph/Classes/AnimGraphNode_SequenceEvaluator.h
@@ -3,7 +3,7 @@
 #pragma once
 #include "AnimGraphNode_Base.h"
 #include "AnimGraphNode_AssetPlayerBase.h"
-#include "Animation/AnimNode_SequenceEvaluator.h"
+#include "AnimNodes/AnimNode_SequenceEvaluator.h"
 #include "EdGraph/EdGraphNodeUtils.h" // for FNodeTextCache
 #include "AnimGraphNode_SequenceEvaluator.generated.h"
 
diff --git a/Engine/Source/Editor/AnimGraph/Classes/AnimGraphNode_Slot.h b/Engine/Source/Editor/AnimGraph/Classes/AnimGraphNode_Slot.h
--- a/Engine/Source/Editor/AnimGraph/Classes/AnimGraphNode_Slot.h
+++ b/Engine/Source/Editor/AnimGraph/Classes/AnimGraphNode_Slot.h
@@ -2,7 +2,7 @@
 
 #pragma once
 #include "AnimGraphNode_Base.h"
-#include "Animation/AnimNode_Slot.h"
+#include "AnimNodes/AnimNode_Slot.h"
 #include "EdGraph/EdGraphNodeUtils.h" // for FNodeTitleTextTable
 #include "AnimGraphNode_Slot.generated.h"
 
diff --git a/Engine/Source/Editor/AnimGraph/Classes/AnimGraphNode_TwoWayBlend.h b/Engine/Source/Editor/AnimGraph/Classes/AnimGraphNode_TwoWayBlend.h
--- a/Engine/Source/Editor/AnimGraph/Classes/AnimGraphNode_TwoWayBlend.h
+++ b/Engine/Source/Editor/AnimGraph/Classes/AnimGraphNode_TwoWayBlend.h
@@ -2,7 +2,7 @@
 
 #pragma once
 #include "AnimGraphNode_Base.h"
-#include "Animation/AnimNode_TwoWayBlend.h"
+#include "AnimNodes/AnimNode_TwoWayBlend.h"
 #include "AnimGraphNode_TwoWayBlend.generated.h"
 
 UCLASS(MinimalAPI)
diff --git a/Engine/Source/Editor/Persona/Private/SSkeletonSlotNames.cpp b/Engine/Source/Editor/Persona/Private/SSkeletonSlotNames.cpp
--- a/Engine/Source/Editor/Persona/Private/SSkeletonSlotNames.cpp
+++ b/Engine/Source/Editor/Persona/Private/SSkeletonSlotNames.cpp
@@ -14,7 +14,7 @@
 #include "STextEntryPopup.h"
 #include "Animation/AnimBlueprint.h"
 #include "AnimGraphNode_Slot.h"
-#include "Animation/AnimNode_Slot.h"
+#include "AnimNodes/AnimNode_Slot.h"
 #include "BlueprintEditorUtils.h"
 #include "SSlotNameReferenceWindow.h"
 #include "IMainFrameModule.h"
diff --git a/Engine/Source/Runtime/AnimGraphRuntime/Private/AnimNodes/AnimNode_ApplyAdditive.cpp b/Engine/Source/Runtime/AnimGraphRuntime/Private/AnimNodes/AnimNode_ApplyAdditive.cpp
new file mode 100644
--- /dev/null
+++ b/Engine/Source/Runtime/AnimGraphRuntime/Private/AnimNodes/AnimNode_ApplyAdditive.cpp
@@ -0,0 +1,68 @@
+// Copyright 1998-2015 Epic Games, Inc. All Rights Reserved.
+
+#include "AnimGraphRuntimePrivatePCH.h"
+#include "AnimNodes/AnimNode_ApplyAdditive.h"
+
+/////////////////////////////////////////////////////
+// FAnimNode_ApplyAdditive
+
+void FAnimNode_ApplyAdditive::Initialize(const FAnimationInitializeContext& Context)
+{
+	Base.Initialize(Context);
+	Additive.Initialize(Context);
+}
+
+void FAnimNode_ApplyAdditive::CacheBones(const FAnimationCacheBonesContext& Context) 
+{
+	Base.CacheBones(Context);
+	Additive.CacheBones(Context);
+}
+
+void FAnimNode_ApplyAdditive::Update(const FAnimationUpdateContext& Context)
+{
+	EvaluateGraphExposedInputs.Execute(Context);
+
+	Base.Update(Context);
+	const float ActualAlpha = AlphaScaleBias.ApplyTo(Alpha);
+	if (ActualAlpha > ZERO_ANIMWEIGHT_THRESH)
+	{
+		Additive.Update(Context.FractionalWeight(ActualAlpha));
+	}
+}
+
+void FAnimNode_ApplyAdditive::Evaluate(FPoseContext& Output)
+{
+	//@TODO: Could evaluate Base into Output and save a copy
+	const float ActualAlpha = AlphaScaleBias.ApplyTo(Alpha);
+	if (ActualAlpha > ZERO_ANIMWEIGHT_THRESH)
+	{
+		FPoseContext AdditiveEvalContext(Output);
+
+		Base.Evaluate(Output);
+		Additive.Evaluate(AdditiveEvalContext);
+
+		FAnimationRuntime::AccumulateAdditivePose(Output.Pose, AdditiveEvalContext.Pose, Output.Curve, AdditiveEvalContext.Curve, ActualAlpha, AAT_LocalSpaceBase);
+		Output.Pose.NormalizeRotations();
+	}
+	else
+	{
+		Base.Evaluate(Output);
+	}
+}
+
+FAnimNode_ApplyAdditive::FAnimNode_ApplyAdditive()
+	: Alpha(1.0f)
+{
+}
+
+void FAnimNode_ApplyAdditive::GatherDebugData(FNodeDebugData& DebugData)
+{
+	const float ActualAlpha = AlphaScaleBias.ApplyTo(Alpha);
+
+	FString DebugLine = DebugData.GetNodeName(this);
+	DebugLine += FString::Printf(TEXT("(Alpha: %.1f%%)"), ActualAlpha*100.f);
+
+	DebugData.AddDebugItem(DebugLine);
+	Base.GatherDebugData(DebugData.BranchFlow(1.f));
+	Additive.GatherDebugData(DebugData.BranchFlow(ActualAlpha));
+}
\ No newline at end of file
diff --git a/Engine/Source/Runtime/AnimGraphRuntime/Private/AnimNodes/AnimNode_BlendListBase.cpp b/Engine/Source/Runtime/AnimGraphRuntime/Private/AnimNodes/AnimNode_BlendListBase.cpp
new file mode 100644
--- /dev/null
+++ b/Engine/Source/Runtime/AnimGraphRuntime/Private/AnimNodes/AnimNode_BlendListBase.cpp
@@ -0,0 +1,198 @@
+// Copyright 1998-2015 Epic Games, Inc. All Rights Reserved.
+
+#include "AnimGraphRuntimePrivatePCH.h"
+#include "AnimNodes/AnimNode_BlendListBase.h"
+#include "AnimationRuntime.h"
+#include "Animation/AnimStats.h"
+
+/////////////////////////////////////////////////////
+// FAnimNode_BlendListBase
+
+void FAnimNode_BlendListBase::Initialize(const FAnimationInitializeContext& Context)
+{
+	const int NumPoses = BlendPose.Num();
+	checkSlow(BlendTime.Num() == NumPoses);
+
+	BlendWeights.Empty(NumPoses);
+	if (NumPoses > 0)
+	{
+		BlendWeights.AddZeroed(NumPoses);
+		BlendWeights[0] = 1.0f;
+
+		for (int32 ChildIndex = 0; ChildIndex < NumPoses; ++ChildIndex)
+		{
+			BlendPose[ChildIndex].Initialize(Context);
+		}
+	}
+
+	RemainingBlendTimes.Empty(NumPoses);
+	RemainingBlendTimes.AddZeroed(NumPoses);
+	Blends.Empty(NumPoses);
+	Blends.AddZeroed(NumPoses);
+
+	LastActiveChildIndex = INDEX_NONE;
+
+	for(int32 i = 0 ; i < Blends.Num() ; ++i)
+	{
+		FAlphaBlend& Blend = Blends[i];
+
+		Blend.SetBlendTime(0.0f);
+		Blend.SetBlendOption(BlendType);
+		Blend.CustomCurve = CustomBlendCurve;
+	}
+	Blends[0].SetAlpha(1.0f);
+}
+
+void FAnimNode_BlendListBase::CacheBones(const FAnimationCacheBonesContext& Context) 
+{
+	for(int32 ChildIndex=0; ChildIndex<BlendPose.Num(); ChildIndex++)
+	{
+		BlendPose[ChildIndex].CacheBones(Context);
+	}
+}
+
+void FAnimNode_BlendListBase::Update(const FAnimationUpdateContext& Context)
+{
+	EvaluateGraphExposedInputs.Execute(Context);
+
+	const int NumPoses = BlendPose.Num();
+	checkSlow((BlendTime.Num() == NumPoses) && (BlendWeights.Num() == NumPoses));
+
+	PosesToEvaluate.Empty(NumPoses);
+
+	if (NumPoses > 0)
+	{
+		// Handle a change in the active child index; adjusting the target weights
+		const int32 ChildIndex = GetActiveChildIndex();
+
+		if (ChildIndex != LastActiveChildIndex)
+		{
+			bool LastChildIndexIsInvalid = (LastActiveChildIndex == INDEX_NONE);
+			LastActiveChildIndex = ChildIndex;
+
+			const float CurrentWeight = BlendWeights[ChildIndex];
+			const float DesiredWeight = 1.0f;
+			const float WeightDifference = FMath::Clamp<float>(FMath::Abs<float>(DesiredWeight - CurrentWeight), 0.0f, 1.0f);
+
+			// scale by the weight difference since we want always consistency:
+			// - if you're moving from 0 to full weight 1, it will use the normal blend time
+			// - if you're moving from 0.5 to full weight 1, it will get there in half the time
+			const float RemainingBlendTime = LastChildIndexIsInvalid ? 0.0f : ( BlendTime[ChildIndex] * WeightDifference );
+
+			for (int32 i = 0; i < RemainingBlendTimes.Num(); ++i)
+			{
+				RemainingBlendTimes[i] = RemainingBlendTime;
+			}
+
+			for(int32 i = 0; i < Blends.Num(); ++i)
+			{
+				FAlphaBlend& Blend = Blends[i];
+
+				Blend.SetBlendTime(RemainingBlendTime);
+
+				if(i == ChildIndex)
+				{
+					Blend.SetValueRange(BlendWeights[i], 1.0f);
+				}
+				else
+				{
+					Blend.SetValueRange(BlendWeights[i], 0.0f);
+				}
+
+				Blend.Reset();
+			}
+		}
+
+		// Advance the weights
+		//@TODO: This means we advance even in a frame where the target weights/times just got modified; is that desirable?
+		float SumWeight = 0.0f;
+		for (int32 i = 0; i < Blends.Num(); ++i)
+		{
+			float& BlendWeight = BlendWeights[i];
+
+			FAlphaBlend& Blend = Blends[i];
+			Blend.Update(Context.GetDeltaTime());
+			BlendWeight = Blend.GetBlendedValue();
+
+			SumWeight += BlendWeight;
+		}
+
+		// Renormalize the weights
+		if ((SumWeight > ZERO_ANIMWEIGHT_THRESH) && (FMath::Abs<float>(SumWeight - 1.0f) > ZERO_ANIMWEIGHT_THRESH))
+		{
+			float ReciprocalSum = 1.0f / SumWeight;
+			for (int32 i = 0; i < BlendWeights.Num(); ++i)
+			{
+				BlendWeights[i] *= ReciprocalSum;
+			}
+		}
+
+		// Update our active children
+		for (int32 i = 0; i < BlendPose.Num(); ++i)
+		{
+			const float BlendWeight = BlendWeights[i];
+			if (BlendWeight > ZERO_ANIMWEIGHT_THRESH)
+			{
+				BlendPose[i].Update(Context.FractionalWeight(BlendWeight));
+				PosesToEvaluate.Add(i);
+			}
+		}
+	}
+}
+
+void FAnimNode_BlendListBase::Evaluate(FPoseContext& Output)
+{
+	ANIM_MT_SCOPE_CYCLE_COUNTER(BlendPosesInGraph, Output.AnimInstance->IsRunningParallelEvaluation());
+
+	const int32 NumPoses = PosesToEvaluate.Num();
+
+	if ((NumPoses > 0) && (BlendPose.Num() == BlendWeights.Num()))
+	{
+		TArray<FCompactPose> FilteredPoses;
+		FilteredPoses.SetNum(NumPoses);
+
+		TArray<FBlendedCurve> FilteredCurve;
+		FilteredCurve.SetNum(NumPoses);
+
+		TArray<float> FilteredWeights;
+		FilteredWeights.AddUninitialized(NumPoses);
+
+		int32 NumActivePoses = 0;
+		for (int32 i = 0; i < PosesToEvaluate.Num(); ++i)
+		{
+			int32 PoseIndex = PosesToEvaluate[i];
+
+			const float BlendWeight = BlendWeights[PoseIndex];
+			FPoseContext EvaluateContext(Output);
+
+			FPoseLink& CurrentPose = BlendPose[PoseIndex];
+			CurrentPose.Evaluate(EvaluateContext);
+
+			FilteredPoses[i].MoveBonesFrom(EvaluateContext.Pose);
+			FilteredCurve[i] = EvaluateContext.Curve;
+			FilteredWeights[i] = BlendWeight;
+		}
+
+		FAnimationRuntime::BlendPosesTogether(FilteredPoses, FilteredCurve, FilteredWeights, Output.Pose, Output.Curve);
+	}
+	else
+	{
+		Output.ResetToRefPose();
+	}
+}
+
+void FAnimNode_BlendListBase::GatherDebugData(FNodeDebugData& DebugData)
+{
+	const int NumPoses = BlendPose.Num();
+	const int32 ChildIndex = GetActiveChildIndex();
+
+	FString DebugLine = GetNodeName(DebugData);
+	DebugLine += FString::Printf(TEXT("(Active: (%i/%i) BlendWeight: %.1f%% BlendTime %.3f)"), ChildIndex+1, NumPoses, BlendWeights[ChildIndex]*100.f, BlendTime[ChildIndex]);
+
+	DebugData.AddDebugItem(DebugLine);
+	
+	for(int32 Pose = 0; Pose < NumPoses; ++Pose)
+	{
+		BlendPose[Pose].GatherDebugData(DebugData.BranchFlow(BlendWeights[Pose]));
+	}
+}
\ No newline at end of file
diff --git a/Engine/Source/Runtime/AnimGraphRuntime/Private/AnimNodes/AnimNode_BlendListByBool.cpp b/Engine/Source/Runtime/AnimGraphRuntime/Private/AnimNodes/AnimNode_BlendListByBool.cpp
new file mode 100644
--- /dev/null
+++ b/Engine/Source/Runtime/AnimGraphRuntime/Private/AnimNodes/AnimNode_BlendListByBool.cpp
@@ -0,0 +1,14 @@
+// Copyright 1998-2015 Epic Games, Inc. All Rights Reserved.
+
+#include "AnimGraphRuntimePrivatePCH.h"
+#include "AnimNodes/AnimNode_BlendListByBool.h"
+#include "AnimationRuntime.h"
+
+/////////////////////////////////////////////////////
+// FAnimNode_BlendListByBool
+
+int32 FAnimNode_BlendListByBool::GetActiveChildIndex()
+{
+	// Note: Intentionally flipped boolean sense (the true input is #0, and the false input is #1)
+	return bActiveValue ? 0 : 1;
+}
diff --git a/Engine/Source/Runtime/AnimGraphRuntime/Private/AnimNodes/AnimNode_BlendListByEnum.cpp b/Engine/Source/Runtime/AnimGraphRuntime/Private/AnimNodes/AnimNode_BlendListByEnum.cpp
new file mode 100644
--- /dev/null
+++ b/Engine/Source/Runtime/AnimGraphRuntime/Private/AnimNodes/AnimNode_BlendListByEnum.cpp
@@ -0,0 +1,21 @@
+// Copyright 1998-2015 Epic Games, Inc. All Rights Reserved.
+
+#include "AnimGraphRuntimePrivatePCH.h"
+#include "AnimationRuntime.h"
+#include "AnimNodes/AnimNode_BlendListByEnum.h"
+
+/////////////////////////////////////////////////////
+// FAnimNode_BlendListByEnum
+
+int32 FAnimNode_BlendListByEnum::GetActiveChildIndex()
+{
+	if (EnumToPoseIndex.IsValidIndex(ActiveEnumValue))
+	{
+		return EnumToPoseIndex[ActiveEnumValue];
+	}
+	else
+	{
+		return 0;
+	}
+}
+
diff --git a/Engine/Source/Runtime/AnimGraphRuntime/Private/AnimNodes/AnimNode_BlendListByInt.cpp b/Engine/Source/Runtime/AnimGraphRuntime/Private/AnimNodes/AnimNode_BlendListByInt.cpp
new file mode 100644
--- /dev/null
+++ b/Engine/Source/Runtime/AnimGraphRuntime/Private/AnimNodes/AnimNode_BlendListByInt.cpp
@@ -0,0 +1,14 @@
+// Copyright 1998-2015 Epic Games, Inc. All Rights Reserved.
+
+#include "AnimGraphRuntimePrivatePCH.h"
+#include "AnimNodes/AnimNode_BlendListByInt.h"
+#include "AnimationRuntime.h"
+
+/////////////////////////////////////////////////////
+// FAnimNode_BlendListByInt
+
+int32 FAnimNode_BlendListByInt::GetActiveChildIndex()
+{
+	const int NumPoses = BlendPose.Num();
+	return FMath::Clamp<int32>(ActiveChildIndex, 0, NumPoses - 1);
+}
diff --git a/Engine/Source/Runtime/AnimGraphRuntime/Private/AnimNodes/AnimNode_BlendSpaceEvaluator.cpp b/Engine/Source/Runtime/AnimGraphRuntime/Private/AnimNodes/AnimNode_BlendSpaceEvaluator.cpp
new file mode 100644
--- /dev/null
+++ b/Engine/Source/Runtime/AnimGraphRuntime/Private/AnimNodes/AnimNode_BlendSpaceEvaluator.cpp
@@ -0,0 +1,32 @@
+// Copyright 1998-2015 Epic Games, Inc. All Rights Reserved.
+
+#include "AnimGraphRuntimePrivatePCH.h"
+#include "AnimNodes/AnimNode_BlendSpaceEvaluator.h"
+#include "Animation/BlendSpaceBase.h"
+
+/////////////////////////////////////////////////////
+// FAnimNode_BlendSpaceEvaluator
+
+FAnimNode_BlendSpaceEvaluator::FAnimNode_BlendSpaceEvaluator()
+	: FAnimNode_BlendSpacePlayer()
+	, NormalizedTime(0.f)
+{
+}
+
+void FAnimNode_BlendSpaceEvaluator::UpdateAssetPlayer(const FAnimationUpdateContext& Context)
+{
+	EvaluateGraphExposedInputs.Execute(Context);
+	InternalTimeAccumulator = NormalizedTime;
+	PlayRate = 0.f;
+
+	UpdateInternal(Context);
+}
+
+
+void FAnimNode_BlendSpaceEvaluator::GatherDebugData(FNodeDebugData& DebugData)
+{
+	FString DebugLine = DebugData.GetNodeName(this);
+
+	DebugLine += FString::Printf(TEXT("('%s' Play Time: %.3f)"), *BlendSpace->GetName(), InternalTimeAccumulator);
+	DebugData.AddDebugItem(DebugLine, true);
+}
diff --git a/Engine/Source/Runtime/AnimGraphRuntime/Private/AnimNodes/AnimNode_BlendSpacePlayer.cpp b/Engine/Source/Runtime/AnimGraphRuntime/Private/AnimNodes/AnimNode_BlendSpacePlayer.cpp
new file mode 100644
--- /dev/null
+++ b/Engine/Source/Runtime/AnimGraphRuntime/Private/AnimNodes/AnimNode_BlendSpacePlayer.cpp
@@ -0,0 +1,114 @@
+// Copyright 1998-2015 Epic Games, Inc. All Rights Reserved.
+
+#include "AnimGraphRuntimePrivatePCH.h"
+#include "AnimNodes/AnimNode_BlendSpacePlayer.h"
+#include "Animation/BlendSpaceBase.h"
+
+/////////////////////////////////////////////////////
+// FAnimNode_BlendSpacePlayer
+
+FAnimNode_BlendSpacePlayer::FAnimNode_BlendSpacePlayer()
+	: X(0.0f)
+	, Y(0.0f)
+	, Z(0.0f)
+	, PlayRate(1.0f)
+	, bLoop(true)
+	, BlendSpace(NULL)
+	, GroupIndex(INDEX_NONE)
+{
+}
+
+void FAnimNode_BlendSpacePlayer::Initialize(const FAnimationInitializeContext& Context)
+{
+	BlendSampleDataCache.Empty();
+	
+	EvaluateGraphExposedInputs.Execute(Context);
+	if(PlayRate >= 0.0f)
+	{
+		InternalTimeAccumulator = 0.0f;
+	}
+	else
+	{
+		// Blend spaces run between 0 and 1
+		InternalTimeAccumulator = 1.0f;
+	}
+
+	if (BlendSpace != NULL)
+	{
+		BlendSpace->InitializeFilter(&BlendFilter);
+	}
+}
+
+void FAnimNode_BlendSpacePlayer::CacheBones(const FAnimationCacheBonesContext& Context) 
+{
+}
+
+void FAnimNode_BlendSpacePlayer::UpdateAssetPlayer(const FAnimationUpdateContext& Context)
+{
+	EvaluateGraphExposedInputs.Execute(Context);
+
+	UpdateInternal(Context);
+}
+
+void FAnimNode_BlendSpacePlayer::UpdateInternal(const FAnimationUpdateContext& Context)
+{
+	if ((BlendSpace != NULL) && (Context.AnimInstance->CurrentSkeleton->IsCompatible(BlendSpace->GetSkeleton())))
+	{
+		// Create a tick record and fill it out
+		FAnimGroupInstance* SyncGroup;
+		FAnimTickRecord& TickRecord = Context.AnimInstance->CreateUninitializedTickRecord(GroupIndex, /*out*/ SyncGroup);
+
+		const FVector BlendInput(X, Y, Z);
+	
+		Context.AnimInstance->MakeBlendSpaceTickRecord(TickRecord, BlendSpace, BlendInput, BlendSampleDataCache, BlendFilter, bLoop, PlayRate, Context.GetFinalBlendWeight(), /*inout*/ InternalTimeAccumulator);
+
+		// Update the sync group if it exists
+		if (SyncGroup != NULL)
+		{
+			SyncGroup->TestTickRecordForLeadership(GroupRole);
+		}
+	}
+}
+
+void FAnimNode_BlendSpacePlayer::Evaluate(FPoseContext& Output)
+{
+	if ((BlendSpace != NULL) && (Output.AnimInstance->CurrentSkeleton->IsCompatible(BlendSpace->GetSkeleton())))
+	{
+		BlendSpace->GetAnimationPose(BlendSampleDataCache, Output.Pose, Output.Curve);
+	}
+	else
+	{
+		Output.ResetToRefPose();
+	}
+}
+
+void FAnimNode_BlendSpacePlayer::OverrideAsset(UAnimationAsset* NewAsset)
+{
+	if(UBlendSpaceBase* NewBlendSpace = Cast<UBlendSpaceBase>(NewAsset))
+	{
+		BlendSpace = NewBlendSpace;
+	}
+}
+
+void FAnimNode_BlendSpacePlayer::GatherDebugData(FNodeDebugData& DebugData)
+{
+	FString DebugLine = DebugData.GetNodeName(this);
+	if (BlendSpace)
+	{
+		DebugLine += FString::Printf(TEXT("('%s' Play Time: %.3f)"), *BlendSpace->GetName(), InternalTimeAccumulator);
+
+		DebugData.AddDebugItem(DebugLine, true);
+
+		//BlendSpace->GetBlendSamples();
+	}
+}
+
+float FAnimNode_BlendSpacePlayer::GetTimeFromEnd(float CurrentTime)
+{
+	return BlendSpace->GetMaxCurrentTime() - CurrentTime;
+}
+
+UAnimationAsset* FAnimNode_BlendSpacePlayer::GetAnimAsset()
+{
+	return BlendSpace;
+}
diff --git a/Engine/Source/Runtime/AnimGraphRuntime/Private/AnimNodes/AnimNode_LayeredBoneBlend.cpp b/Engine/Source/Runtime/AnimGraphRuntime/Private/AnimNodes/AnimNode_LayeredBoneBlend.cpp
new file mode 100644
--- /dev/null
+++ b/Engine/Source/Runtime/AnimGraphRuntime/Private/AnimNodes/AnimNode_LayeredBoneBlend.cpp
@@ -0,0 +1,166 @@
+// Copyright 1998-2015 Epic Games, Inc. All Rights Reserved.
+
+#include "AnimGraphRuntimePrivatePCH.h"
+#include "AnimNodes/AnimNode_LayeredBoneBlend.h"
+#include "AnimationRuntime.h"
+#include "Animation/AnimStats.h"
+
+/////////////////////////////////////////////////////
+// FAnimNode_LayeredBoneBlend
+
+void FAnimNode_LayeredBoneBlend::Initialize(const FAnimationInitializeContext& Context)
+{
+	const int NumPoses = BlendPoses.Num();
+	checkSlow(BlendWeights.Num() == NumPoses);
+
+	// initialize children
+	BasePose.Initialize(Context);
+
+	if ( NumPoses > 0 )
+	{
+		for (int32 ChildIndex = 0; ChildIndex < NumPoses; ++ChildIndex)
+		{
+			BlendPoses[ChildIndex].Initialize(Context);
+		}
+
+		// initialize mask weight now
+		check (Context.AnimInstance->CurrentSkeleton);
+		ReinitializeBoneBlendWeights(Context.AnimInstance->RequiredBones, Context.AnimInstance->CurrentSkeleton);
+	}
+}
+
+void FAnimNode_LayeredBoneBlend::ReinitializeBoneBlendWeights(const FBoneContainer& RequiredBones, const USkeleton* Skeleton)
+{
+	FAnimationRuntime::CreateMaskWeights(DesiredBoneBlendWeights, LayerSetup, RequiredBones, Skeleton);
+
+	CurrentBoneBlendWeights.Empty(DesiredBoneBlendWeights.Num());
+	CurrentBoneBlendWeights.AddZeroed(DesiredBoneBlendWeights.Num());
+}
+
+void FAnimNode_LayeredBoneBlend::CacheBones(const FAnimationCacheBonesContext& Context) 
+{
+	BasePose.CacheBones(Context);
+	for(int32 ChildIndex=0; ChildIndex<BlendPoses.Num(); ChildIndex++)
+	{
+		BlendPoses[ChildIndex].CacheBones(Context);
+	}
+
+	if (Context.AnimInstance->RequiredBones.GetBoneIndicesArray().Num() != DesiredBoneBlendWeights.Num())
+	{
+		ReinitializeBoneBlendWeights(Context.AnimInstance->RequiredBones, Context.AnimInstance->CurrentSkeleton);
+	}
+}
+
+void FAnimNode_LayeredBoneBlend::Update(const FAnimationUpdateContext& Context)
+{
+	EvaluateGraphExposedInputs.Execute(Context);
+
+	// initialize children
+	BasePose.Update(Context);
+
+	const int NumPoses = BlendPoses.Num();
+	if ( NumPoses > 0 )
+	{
+		for (int32 ChildIndex = 0; ChildIndex < NumPoses; ++ChildIndex)
+		{
+			BlendPoses[ChildIndex].Update(Context);
+		}
+	}
+}
+
+void FAnimNode_LayeredBoneBlend::Evaluate(FPoseContext& Output)
+{
+	ANIM_MT_SCOPE_CYCLE_COUNTER(BlendPosesInGraph, Output.AnimInstance->IsRunningParallelEvaluation());
+
+	const int NumPoses = BlendPoses.Num();
+	if (NumPoses == 0)
+	{
+		BasePose.Evaluate(Output);
+	}
+	else
+	{
+		FAnimationRuntime::UpdateDesiredBoneWeight(DesiredBoneBlendWeights, CurrentBoneBlendWeights, BlendWeights);
+
+		FPoseContext BasePoseContext(Output);
+
+		// evaluate children
+		BasePose.Evaluate(BasePoseContext);
+
+		TArray<FCompactPose> TargetBlendPoses;
+		TargetBlendPoses.SetNum(NumPoses);
+
+		TArray<FBlendedCurve> TargetBlendCurves;
+		TargetBlendCurves.SetNum(NumPoses);
+
+		for (int32 ChildIndex = 0; ChildIndex < NumPoses; ++ChildIndex)
+		{
+			if (BlendWeights[ChildIndex] > ZERO_ANIMWEIGHT_THRESH)
+			{
+				FPoseContext CurrentPoseContext(Output);
+				BlendPoses[ChildIndex].Evaluate(CurrentPoseContext);
+
+				TargetBlendPoses[ChildIndex].MoveBonesFrom(CurrentPoseContext.Pose);
+				TargetBlendCurves[ChildIndex].MoveFrom(CurrentPoseContext.Curve);
+			}
+			else
+			{
+				TargetBlendPoses[ChildIndex].ResetToRefPose(BasePoseContext.Pose.GetBoneContainer());
+				TargetBlendCurves[ChildIndex].InitFrom(Output.Curve);
+			}
+		}
+
+		FAnimationRuntime::BlendPosesPerBoneFilter(BasePoseContext.Pose, TargetBlendPoses, BasePoseContext.Curve, TargetBlendCurves, Output.Pose, Output.Curve, CurrentBoneBlendWeights, bMeshSpaceRotationBlend, CurveBlendOption);
+	}
+}
+
+
+void FAnimNode_LayeredBoneBlend::GatherDebugData(FNodeDebugData& DebugData)
+{
+	const int NumPoses = BlendPoses.Num();
+
+	FString DebugLine = DebugData.GetNodeName(this);
+	DebugLine += FString::Printf(TEXT("(Num Poses: %i)"), NumPoses);
+	DebugData.AddDebugItem(DebugLine);
+
+	BasePose.GatherDebugData(DebugData.BranchFlow(1.f));
+	
+	for (int32 ChildIndex = 0; ChildIndex < NumPoses; ++ChildIndex)
+	{
+		BlendPoses[ChildIndex].GatherDebugData(DebugData.BranchFlow(BlendWeights[ChildIndex]));
+	}
+}
+
+#if WITH_EDITOR
+void FAnimNode_LayeredBoneBlend::ValidateData()
+{
+	// ideally you don't like to get to situation where it becomes inconsistent, but this happened, 
+	// and we don't know what caused this. Possibly copy/paste, but I tried copy/paste and that didn't work
+	// so here we add code to fix this up manually in editor, so that they can continue working on it. 
+	int32 PoseNum = BlendPoses.Num();
+	int32 WeightNum = BlendWeights.Num();
+	int32 LayerNum = LayerSetup.Num();
+
+	int32 Max = FMath::Max3(PoseNum, WeightNum, LayerNum);
+	int32 Min = FMath::Min3(PoseNum, WeightNum, LayerNum);
+	// if they are not all same
+	if (Min != Max)
+	{
+		// we'd like to increase to all Max
+		// sadly we don't have add X for how many
+		for (int32 Index=PoseNum; Index<Max; ++Index)
+		{
+			BlendPoses.Add(FPoseLink());
+		}
+
+		for(int32 Index=WeightNum; Index<Max; ++Index)
+		{
+			BlendWeights.Add(1.f);
+		}
+
+		for(int32 Index=LayerNum; Index<Max; ++Index)
+		{
+			LayerSetup.Add(FInputBlendPose());
+		}
+	}
+}
+#endif
\ No newline at end of file
diff --git a/Engine/Source/Runtime/AnimGraphRuntime/Private/AnimNodes/AnimNode_RefPose.cpp b/Engine/Source/Runtime/AnimGraphRuntime/Private/AnimNodes/AnimNode_RefPose.cpp
new file mode 100644
--- /dev/null
+++ b/Engine/Source/Runtime/AnimGraphRuntime/Private/AnimNodes/AnimNode_RefPose.cpp
@@ -0,0 +1,51 @@
+// Copyright 1998-2015 Epic Games, Inc. All Rights Reserved.
+
+#include "AnimGraphRuntimePrivatePCH.h"
+#include "AnimNodes/AnimNode_RefPose.h"
+
+/////////////////////////////////////////////////////
+// FAnimRefPoseNode
+
+void FAnimNode_RefPose::Evaluate(FPoseContext& Output)
+{
+	// I don't have anything to evaluate. Should this be even here?
+	// EvaluateGraphExposedInputs.Execute(Context);
+
+	switch (RefPoseType) 
+	{
+	case EIT_LocalSpace:
+		Output.ResetToRefPose();
+		break;
+
+	case EIT_Additive:
+	default:
+		Output.ResetToIdentity();
+		break;
+	}
+}
+
+void FAnimNode_MeshSpaceRefPose::EvaluateComponentSpace(FComponentSpacePoseContext& Output)
+{
+	Output.ResetToRefPose();
+}
+
+/** Helper for enum output... */
+#ifndef CASE_ENUM_TO_TEXT
+#define CASE_ENUM_TO_TEXT(txt) case txt: return TEXT(#txt);
+#endif
+
+const TCHAR* GetRefPostTypeText(ERefPoseType RefPose)
+{
+	switch (RefPose)
+	{
+		FOREACH_ENUM_EREFPOSETYPE(CASE_ENUM_TO_TEXT)
+	}
+	return TEXT("Unknown Ref Pose Type");
+}
+
+void FAnimNode_RefPose::GatherDebugData(FNodeDebugData& DebugData)
+{
+	FString DebugLine = DebugData.GetNodeName(this);
+	DebugLine += FString::Printf(TEXT("(Ref Pose Type: %s)"), GetRefPostTypeText(RefPoseType));
+	DebugData.AddDebugItem(DebugLine, true);
+}
\ No newline at end of file
diff --git a/Engine/Source/Runtime/AnimGraphRuntime/Private/AnimNodes/AnimNode_Root.cpp b/Engine/Source/Runtime/AnimGraphRuntime/Private/AnimNodes/AnimNode_Root.cpp
new file mode 100644
--- /dev/null
+++ b/Engine/Source/Runtime/AnimGraphRuntime/Private/AnimNodes/AnimNode_Root.cpp
@@ -0,0 +1,39 @@
+// Copyright 1998-2015 Epic Games, Inc. All Rights Reserved.
+
+#include "AnimGraphRuntimePrivatePCH.h"
+#include "AnimNodes/AnimNode_Root.h"
+
+/////////////////////////////////////////////////////
+// FAnimNode_Root
+
+FAnimNode_Root::FAnimNode_Root()
+{
+}
+
+void FAnimNode_Root::Initialize(const FAnimationInitializeContext& Context)
+{
+	Result.Initialize(Context);
+}
+
+void FAnimNode_Root::CacheBones(const FAnimationCacheBonesContext& Context) 
+{
+	Result.CacheBones(Context);
+}
+
+void FAnimNode_Root::Update(const FAnimationUpdateContext& Context)
+{
+	EvaluateGraphExposedInputs.Execute(Context);
+	Result.Update(Context);
+}
+
+void FAnimNode_Root::Evaluate(FPoseContext& Output)
+{
+	Result.Evaluate(Output);
+}
+
+void FAnimNode_Root::GatherDebugData(FNodeDebugData& DebugData)
+{
+	FString DebugLine = DebugData.GetNodeName(this);
+	DebugData.AddDebugItem(DebugLine);
+	Result.GatherDebugData(DebugData);
+}
\ No newline at end of file
diff --git a/Engine/Source/Runtime/AnimGraphRuntime/Private/AnimNodes/AnimNode_RotateRootBone.cpp b/Engine/Source/Runtime/AnimGraphRuntime/Private/AnimNodes/AnimNode_RotateRootBone.cpp
new file mode 100644
--- /dev/null
+++ b/Engine/Source/Runtime/AnimGraphRuntime/Private/AnimNodes/AnimNode_RotateRootBone.cpp
@@ -0,0 +1,65 @@
+// Copyright 1998-2015 Epic Games, Inc. All Rights Reserved.
+
+#include "AnimGraphRuntimePrivatePCH.h"
+#include "AnimNodes/AnimNode_RotateRootBone.h"
+
+/////////////////////////////////////////////////////
+// FAnimNode_RotateRootBone
+
+void FAnimNode_RotateRootBone::Initialize(const FAnimationInitializeContext& Context)
+{
+	BasePose.Initialize(Context);
+}
+
+void FAnimNode_RotateRootBone::CacheBones(const FAnimationCacheBonesContext& Context) 
+{
+	BasePose.CacheBones(Context);
+}
+
+void FAnimNode_RotateRootBone::Update(const FAnimationUpdateContext& Context)
+{
+	EvaluateGraphExposedInputs.Execute(Context);
+	BasePose.Update(Context);
+}
+
+void FAnimNode_RotateRootBone::Evaluate(FPoseContext& Output)
+{
+	// Evaluate the input
+	BasePose.Evaluate(Output);
+
+	checkSlow(!FMath::IsNaN(Yaw) && FMath::IsFinite(Yaw));
+	checkSlow(!FMath::IsNaN(Pitch) && FMath::IsFinite(Pitch));
+
+	if (Pitch != 0.f || Yaw != 0.f)
+	{
+		// Build our desired rotation
+		const FRotator DeltaRotation(Pitch, Yaw, 0.f);
+		const FQuat DeltaQuat(DeltaRotation);
+		const FQuat MeshToComponentQuat(MeshToComponent);
+
+		// Convert our rotation from Component Space to Mesh Space.
+		const FQuat MeshSpaceDeltaQuat = MeshToComponentQuat.Inverse() * DeltaQuat * MeshToComponentQuat;
+
+		// Apply rotation to root bone.
+		FCompactPoseBoneIndex RootBone(0);
+		Output.Pose[RootBone].SetRotation(Output.Pose[RootBone].GetRotation() * MeshSpaceDeltaQuat);
+	}
+}
+
+
+void FAnimNode_RotateRootBone::GatherDebugData(FNodeDebugData& DebugData)
+{
+	FString DebugLine = DebugData.GetNodeName(this);
+	
+	DebugLine += FString::Printf(TEXT("(Pitch: %.2f Yaw: %.2f)"), Pitch, Yaw);
+	DebugData.AddDebugItem(DebugLine);
+
+	BasePose.GatherDebugData(DebugData);
+}
+
+FAnimNode_RotateRootBone::FAnimNode_RotateRootBone()
+	: Pitch(0.0f)
+	, Yaw(0.0f)
+	, MeshToComponent(FRotator::ZeroRotator)
+{
+}
diff --git a/Engine/Source/Runtime/AnimGraphRuntime/Private/AnimNodes/AnimNode_RotationOffsetBlendSpace.cpp b/Engine/Source/Runtime/AnimGraphRuntime/Private/AnimNodes/AnimNode_RotationOffsetBlendSpace.cpp
new file mode 100644
--- /dev/null
+++ b/Engine/Source/Runtime/AnimGraphRuntime/Private/AnimNodes/AnimNode_RotationOffsetBlendSpace.cpp
@@ -0,0 +1,55 @@
+// Copyright 1998-2015 Epic Games, Inc. All Rights Reserved.
+
+#include "AnimGraphRuntimePrivatePCH.h"
+#include "AnimNodes/AnimNode_RotationOffsetBlendSpace.h"
+
+/////////////////////////////////////////////////////
+// FAnimNode_RotationOffsetBlendSpace
+
+void FAnimNode_RotationOffsetBlendSpace::Initialize(const FAnimationInitializeContext& Context)
+{
+	FAnimNode_BlendSpacePlayer::Initialize(Context);
+	BasePose.Initialize(Context);
+}
+
+void FAnimNode_RotationOffsetBlendSpace::CacheBones(const FAnimationCacheBonesContext& Context) 
+{
+	FAnimNode_BlendSpacePlayer::CacheBones(Context);
+	BasePose.CacheBones(Context);
+}
+
+void FAnimNode_RotationOffsetBlendSpace::UpdateAssetPlayer(const FAnimationUpdateContext& Context)
+{
+	FAnimNode_BlendSpacePlayer::UpdateAssetPlayer(Context);
+	BasePose.Update(Context);
+}
+
+void FAnimNode_RotationOffsetBlendSpace::Evaluate(FPoseContext& Context)
+{
+	// Evaluate base pose
+	BasePose.Evaluate(Context);
+
+	// Evaluate MeshSpaceRotation additive blendspace
+	FPoseContext MeshSpaceRotationAdditivePoseContext(Context);
+	FAnimNode_BlendSpacePlayer::Evaluate(MeshSpaceRotationAdditivePoseContext);
+
+	// Accumulate poses together
+	FAnimationRuntime::AccumulateMeshSpaceRotationAdditiveToLocalPose(Context.Pose, MeshSpaceRotationAdditivePoseContext.Pose, Context.Curve, MeshSpaceRotationAdditivePoseContext.Curve, 1.f);
+
+	// Resulting rotations are not normalized, so normalize here.
+	Context.Pose.NormalizeRotations();
+}
+
+void FAnimNode_RotationOffsetBlendSpace::GatherDebugData(FNodeDebugData& DebugData)
+{
+	FString DebugLine = DebugData.GetNodeName(this);
+	
+	DebugLine += FString::Printf(TEXT("(Play Time: %.3f)"), InternalTimeAccumulator);
+	DebugData.AddDebugItem(DebugLine);
+	
+	BasePose.GatherDebugData(DebugData);
+}
+
+FAnimNode_RotationOffsetBlendSpace::FAnimNode_RotationOffsetBlendSpace()
+{
+}
diff --git a/Engine/Source/Runtime/AnimGraphRuntime/Private/AnimNodes/AnimNode_SaveCachedPose.cpp b/Engine/Source/Runtime/AnimGraphRuntime/Private/AnimNodes/AnimNode_SaveCachedPose.cpp
new file mode 100644
--- /dev/null
+++ b/Engine/Source/Runtime/AnimGraphRuntime/Private/AnimNodes/AnimNode_SaveCachedPose.cpp
@@ -0,0 +1,74 @@
+// Copyright 1998-2015 Epic Games, Inc. All Rights Reserved.
+
+#include "AnimGraphRuntimePrivatePCH.h"
+#include "AnimNodes/AnimNode_SaveCachedPose.h"
+
+/////////////////////////////////////////////////////
+// FAnimNode_SaveCachedPose
+
+FAnimNode_SaveCachedPose::FAnimNode_SaveCachedPose()
+	: LastInitializedContextCounter(INDEX_NONE)
+	, LastCacheBonesContextCounter(INDEX_NONE)
+	, LastUpdatedContextCounter(INDEX_NONE)
+	, LastEvaluatedContextCounter(INDEX_NONE)
+{
+}
+
+void FAnimNode_SaveCachedPose::Initialize(const FAnimationInitializeContext& Context)
+{
+	if (LastInitializedContextCounter != Context.AnimInstance->GetGraphTraversalCounter())
+	{
+		LastInitializedContextCounter = Context.AnimInstance->GetGraphTraversalCounter();
+
+		// Initialize the subgraph
+		Pose.Initialize(Context);
+	}
+}
+
+void FAnimNode_SaveCachedPose::CacheBones(const FAnimationCacheBonesContext& Context)
+{
+	if (LastCacheBonesContextCounter != Context.AnimInstance->GetGraphTraversalCounter())
+	{
+		LastCacheBonesContextCounter = Context.AnimInstance->GetGraphTraversalCounter();
+
+		// Initialize the subgraph
+		Pose.CacheBones(Context);
+	}
+}
+
+void FAnimNode_SaveCachedPose::Update(const FAnimationUpdateContext& Context)
+{
+	if (LastUpdatedContextCounter != Context.AnimInstance->GetGraphTraversalCounter())
+	{
+		LastUpdatedContextCounter = Context.AnimInstance->GetGraphTraversalCounter();
+
+		// Update the subgraph
+		Pose.Update(Context);
+	}
+}
+
+void FAnimNode_SaveCachedPose::Evaluate(FPoseContext& Output)
+{
+	if (LastEvaluatedContextCounter != Output.AnimInstance->GetGraphTraversalCounter())
+	{
+		LastEvaluatedContextCounter = Output.AnimInstance->GetGraphTraversalCounter();
+
+		FPoseContext CachingContext(Output);
+		Pose.Evaluate(CachingContext);
+		CachedPose.MoveBonesFrom(CachingContext.Pose);
+		CachedCurve.MoveFrom(CachingContext.Curve);
+	}
+
+	// Return the cached result
+	Output.Pose.CopyBonesFrom(CachedPose);
+	Output.Curve.CopyFrom(CachedCurve);
+}
+
+
+void FAnimNode_SaveCachedPose::GatherDebugData(FNodeDebugData& DebugData)
+{
+	FString DebugLine = DebugData.GetNodeName(this);
+	DebugData.AddDebugItem(DebugLine);
+
+	Pose.GatherDebugData(DebugData);
+}
\ No newline at end of file
diff --git a/Engine/Source/Runtime/AnimGraphRuntime/Private/AnimNodes/AnimNode_SequenceEvaluator.cpp b/Engine/Source/Runtime/AnimGraphRuntime/Private/AnimNodes/AnimNode_SequenceEvaluator.cpp
new file mode 100644
--- /dev/null
+++ b/Engine/Source/Runtime/AnimGraphRuntime/Private/AnimNodes/AnimNode_SequenceEvaluator.cpp
@@ -0,0 +1,47 @@
+// Copyright 1998-2015 Epic Games, Inc. All Rights Reserved.
+
+#include "AnimGraphRuntimePrivatePCH.h"
+#include "AnimNodes/AnimNode_SequenceEvaluator.h"
+
+/////////////////////////////////////////////////////
+// FAnimSequenceEvaluatorNode
+
+void FAnimNode_SequenceEvaluator::Initialize(const FAnimationInitializeContext& Context)
+{
+}
+
+void FAnimNode_SequenceEvaluator::CacheBones(const FAnimationCacheBonesContext& Context) 
+{
+}
+
+void FAnimNode_SequenceEvaluator::UpdateAssetPlayer(const FAnimationUpdateContext& Context)
+{
+	EvaluateGraphExposedInputs.Execute(Context);
+}
+
+void FAnimNode_SequenceEvaluator::Evaluate(FPoseContext& Output)
+{
+	if ((Sequence != NULL) && (Output.AnimInstance->CurrentSkeleton->IsCompatible(Sequence->GetSkeleton())))
+	{
+		Sequence->GetAnimationPose(Output.Pose, Output.Curve, FAnimExtractContext(ExplicitTime, Output.AnimInstance->ShouldExtractRootMotion()));
+	}
+	else
+	{
+		Output.ResetToRefPose();
+	}
+}
+void FAnimNode_SequenceEvaluator::OverrideAsset(UAnimationAsset* NewAsset)
+{
+	if(UAnimSequenceBase* NewSequence = Cast<UAnimSequenceBase>(NewAsset))
+	{
+		Sequence = NewSequence;
+	}
+}
+
+void FAnimNode_SequenceEvaluator::GatherDebugData(FNodeDebugData& DebugData)
+{
+	FString DebugLine = DebugData.GetNodeName(this);
+	
+	DebugLine += FString::Printf(TEXT("('%s' Play Time: %.3f)"), *Sequence->GetName(), ExplicitTime);
+	DebugData.AddDebugItem(DebugLine, true);
+}
\ No newline at end of file
diff --git a/Engine/Source/Runtime/AnimGraphRuntime/Private/AnimNodes/AnimNode_Slot.cpp b/Engine/Source/Runtime/AnimGraphRuntime/Private/AnimNodes/AnimNode_Slot.cpp
new file mode 100644
--- /dev/null
+++ b/Engine/Source/Runtime/AnimGraphRuntime/Private/AnimNodes/AnimNode_Slot.cpp
@@ -0,0 +1,103 @@
+// Copyright 1998-2015 Epic Games, Inc. All Rights Reserved.
+
+#include "AnimGraphRuntimePrivatePCH.h"
+#include "AnimNodes/AnimNode_Slot.h"
+#include "Animation/AnimMontage.h"
+
+/////////////////////////////////////////////////////
+// FAnimNode_Slot
+
+void FAnimNode_Slot::Initialize(const FAnimationInitializeContext& Context)
+{
+	Source.Initialize(Context);
+
+	SourceWeight = 0.f;
+	SlotNodeWeight = 0.f;
+
+	// If this node has not already been registered with the AnimInstance, do it.
+	if ((LastSlotNodeInitializationCounter == INDEX_NONE) || (LastSlotNodeInitializationCounter != Context.AnimInstance->GetSlotNodeInitializationCounter()))
+	{
+		LastSlotNodeInitializationCounter = Context.AnimInstance->GetSlotNodeInitializationCounter();
+		Context.AnimInstance->RegisterSlotNodeWithAnimInstance(SlotName);
+	}
+}
+
+void FAnimNode_Slot::CacheBones(const FAnimationCacheBonesContext& Context) 
+{
+	Source.CacheBones(Context);
+}
+
+void FAnimNode_Slot::Update(const FAnimationUpdateContext& Context)
+{
+	// Update weights.
+	Context.AnimInstance->GetSlotWeight(SlotName, SlotNodeWeight, SourceWeight);
+
+	// Update cache in AnimInstance.
+	Context.AnimInstance->UpdateSlotNodeWeight(SlotName, SlotNodeWeight);
+	Context.AnimInstance->UpdateSlotRootMotionWeight(SlotName, Context.GetFinalBlendWeight());
+
+
+	if (SourceWeight > ZERO_ANIMWEIGHT_THRESH)
+	{
+		Source.Update(Context.FractionalWeight(SourceWeight));
+	}
+}
+
+void FAnimNode_Slot::Evaluate(FPoseContext & Output)
+{
+	// If not playing a montage, just pass through
+	if (SlotNodeWeight <= ZERO_ANIMWEIGHT_THRESH)
+	{
+		Source.Evaluate(Output);
+	}
+	else
+	{
+		FPoseContext SourceContext(Output);
+		if (SourceWeight > ZERO_ANIMWEIGHT_THRESH)
+		{
+			Source.Evaluate(SourceContext);
+		}
+
+		Output.AnimInstance->SlotEvaluatePose(SlotName, SourceContext.Pose, SourceContext.Curve, Output.Pose, Output.Curve, SlotNodeWeight);
+
+		checkSlow(!Output.ContainsNaN());
+		checkSlow(Output.IsNormalized());
+	}
+}
+
+void FAnimNode_Slot::GatherDebugData(FNodeDebugData& DebugData)
+{
+	FString DebugLine = DebugData.GetNodeName(this);
+	DebugLine += FString::Printf(TEXT("(Slot Name: '%s' Weight:%.1f%%)"), *SlotName.ToString(), SlotNodeWeight*100.f);
+	
+	bool const bIsPoseSource = (SourceWeight <= ZERO_ANIMWEIGHT_THRESH);
+	DebugData.AddDebugItem(DebugLine, bIsPoseSource);
+	Source.GatherDebugData(DebugData.BranchFlow(SourceWeight));
+
+	for (FAnimMontageInstance* MontageInstance : DebugData.AnimInstance->MontageInstances)
+	{
+		if (MontageInstance->IsValid() && MontageInstance->Montage->IsValidSlot(SlotName))
+		{
+			const FAnimTrack* Track = MontageInstance->Montage->GetAnimationData(SlotName);
+			for (const FAnimSegment& Segment : Track->AnimSegments)
+			{
+				float Weight;
+				float CurrentAnimPos;
+				if (UAnimSequenceBase* Anim = Segment.GetAnimationData(MontageInstance->GetPosition(), CurrentAnimPos, Weight))
+				{
+					FString MontageLine = FString::Printf(TEXT("Montage: '%s' Anim: '%s' Play Time:%.2f W:%.2f%%"), *MontageInstance->Montage->GetName(), *Anim->GetName(), CurrentAnimPos, Weight*100.f);
+					DebugData.BranchFlow(1.0f).AddDebugItem(MontageLine, true);
+					break;
+				}
+			}
+		}
+	}
+}
+
+FAnimNode_Slot::FAnimNode_Slot()
+	: SlotName(FAnimSlotGroup::DefaultSlotName)
+	, SourceWeight(0.f)
+	, SlotNodeWeight(0.f)
+	, LastSlotNodeInitializationCounter(INDEX_NONE)
+{
+}
diff --git a/Engine/Source/Runtime/AnimGraphRuntime/Private/AnimNodes/AnimNode_TwoWayBlend.cpp b/Engine/Source/Runtime/AnimGraphRuntime/Private/AnimNodes/AnimNode_TwoWayBlend.cpp
new file mode 100644
--- /dev/null
+++ b/Engine/Source/Runtime/AnimGraphRuntime/Private/AnimNodes/AnimNode_TwoWayBlend.cpp
@@ -0,0 +1,84 @@
+// Copyright 1998-2015 Epic Games, Inc. All Rights Reserved.
+
+#include "AnimGraphRuntimePrivatePCH.h"
+#include "AnimNodes/AnimNode_TwoWayBlend.h"
+
+/////////////////////////////////////////////////////
+// FAnimationNode_TwoWayBlend
+
+void FAnimationNode_TwoWayBlend::Initialize(const FAnimationInitializeContext& Context)
+{
+	A.Initialize(Context);
+	B.Initialize(Context);
+}
+
+void FAnimationNode_TwoWayBlend::CacheBones(const FAnimationCacheBonesContext& Context) 
+{
+	A.CacheBones(Context);
+	B.CacheBones(Context);
+}
+
+void FAnimationNode_TwoWayBlend::Update(const FAnimationUpdateContext& Context)
+{
+	EvaluateGraphExposedInputs.Execute(Context);
+
+	const float ActualAlpha = AlphaScaleBias.ApplyTo(Alpha);
+	if (ActualAlpha > ZERO_ANIMWEIGHT_THRESH)
+	{
+		if (ActualAlpha < 1.0f - ZERO_ANIMWEIGHT_THRESH)
+		{
+			// Blend A and B together
+			A.Update(Context.FractionalWeight(1.0f - ActualAlpha));
+			B.Update(Context.FractionalWeight(ActualAlpha));
+		}
+		else
+		{
+			// Take all of B
+			B.Update(Context);
+		}
+	}
+	else
+	{
+		// Take all of A
+		A.Update(Context);
+	}
+}
+
+void FAnimationNode_TwoWayBlend::Evaluate(FPoseContext& Output)
+{
+	const float ActualAlpha = AlphaScaleBias.ApplyTo(Alpha);
+	if (ActualAlpha > ZERO_ANIMWEIGHT_THRESH)
+	{
+		if (ActualAlpha < 1.0f - ZERO_ANIMWEIGHT_THRESH)
+		{
+			FPoseContext Pose1(Output);
+			FPoseContext Pose2(Output);
+
+			A.Evaluate(Pose1);
+			B.Evaluate(Pose2);
+
+			FAnimationRuntime::BlendTwoPosesTogether(Pose1.Pose, Pose2.Pose, Pose1.Curve, Pose2.Curve, (1.0f - ActualAlpha), Output.Pose, Output.Curve);
+		}
+		else
+		{
+			B.Evaluate(Output);
+		}
+	}
+	else
+	{
+		A.Evaluate(Output);
+	}
+}
+
+
+void FAnimationNode_TwoWayBlend::GatherDebugData(FNodeDebugData& DebugData)
+{
+	const float ActualAlpha = AlphaScaleBias.ApplyTo(Alpha);
+
+	FString DebugLine = DebugData.GetNodeName(this);
+	DebugLine += FString::Printf(TEXT("(Alpha: %.1f%%)"), ActualAlpha*100);
+	DebugData.AddDebugItem(DebugLine);
+
+	A.GatherDebugData(DebugData.BranchFlow(1.f - ActualAlpha));
+	B.GatherDebugData(DebugData.BranchFlow(ActualAlpha));
+}
\ No newline at end of file
diff --git a/Engine/Source/Runtime/AnimGraphRuntime/Public/AnimNodes/AnimNode_ApplyAdditive.h b/Engine/Source/Runtime/AnimGraphRuntime/Public/AnimNodes/AnimNode_ApplyAdditive.h
new file mode 100644
--- /dev/null
+++ b/Engine/Source/Runtime/AnimGraphRuntime/Public/AnimNodes/AnimNode_ApplyAdditive.h
@@ -0,0 +1,36 @@
+// Copyright 1998-2015 Epic Games, Inc. All Rights Reserved.
+
+#pragma once
+
+#include "Animation/AnimNodeBase.h"
+#include "Animation/InputScaleBias.h"
+#include "AnimNode_ApplyAdditive.generated.h"
+
+USTRUCT()
+struct ANIMGRAPHRUNTIME_API FAnimNode_ApplyAdditive : public FAnimNode_Base
+{
+	GENERATED_USTRUCT_BODY()
+
+	UPROPERTY(EditAnywhere, BlueprintReadWrite, Category=Links)
+	FPoseLink Base;
+
+	UPROPERTY(EditAnywhere, BlueprintReadWrite, Category=Links)
+	FPoseLink Additive;
+
+	UPROPERTY(EditAnywhere, BlueprintReadWrite, Category=Settings, meta=(PinShownByDefault))
+	mutable float Alpha;
+
+	UPROPERTY(EditAnywhere, BlueprintReadWrite, Category=Settings)
+	FInputScaleBias AlphaScaleBias;
+
+public:	
+	FAnimNode_ApplyAdditive();
+
+	// FAnimNode_Base interface
+	virtual void Initialize(const FAnimationInitializeContext& Context) override;
+	virtual void CacheBones(const FAnimationCacheBonesContext& Context) override;
+	virtual void Update(const FAnimationUpdateContext& Context) override;
+	virtual void Evaluate(FPoseContext& Output) override;
+	virtual void GatherDebugData(FNodeDebugData& DebugData) override;
+	// End of FAnimNode_Base interface
+};
diff --git a/Engine/Source/Runtime/AnimGraphRuntime/Public/AnimNodes/AnimNode_BlendListBase.h b/Engine/Source/Runtime/AnimGraphRuntime/Public/AnimNodes/AnimNode_BlendListBase.h
new file mode 100644
--- /dev/null
+++ b/Engine/Source/Runtime/AnimGraphRuntime/Public/AnimNodes/AnimNode_BlendListBase.h
@@ -0,0 +1,76 @@
+// Copyright 1998-2015 Epic Games, Inc. All Rights Reserved.
+
+#pragma once
+
+#include "AlphaBlend.h"
+#include "Animation/AnimNodeBase.h"
+#include "AnimNode_BlendListBase.generated.h"
+
+// Blend list node; has many children
+USTRUCT()
+struct ANIMGRAPHRUNTIME_API FAnimNode_BlendListBase : public FAnimNode_Base
+{
+	GENERATED_USTRUCT_BODY()
+public:
+	UPROPERTY(EditAnywhere, EditFixedSize, BlueprintReadWrite, Category=Links)
+	TArray<FPoseLink> BlendPose;
+
+	UPROPERTY(EditAnywhere, EditFixedSize, BlueprintReadWrite, Category=Config, meta=(PinShownByDefault))
+	TArray<float> BlendTime;
+
+	UPROPERTY(EditAnywhere, Category=BlendType)
+	EAlphaBlendOption BlendType;
+
+	UPROPERTY(EditAnywhere, Category=BlendType)
+	UCurveFloat* CustomBlendCurve;
+
+	UPROPERTY()
+	TArray<struct FAlphaBlend> Blends;
+
+protected:
+	UPROPERTY()
+	TArray<float> BlendWeights;
+
+	UPROPERTY()
+	TArray<float> RemainingBlendTimes;
+
+	UPROPERTY()
+	int32 LastActiveChildIndex;
+
+	//Store which poses we need to evaluate
+	TArray<int32> PosesToEvaluate;
+
+public:	
+	FAnimNode_BlendListBase()
+		: LastActiveChildIndex(0)
+	{
+		BlendWeights.Empty();
+		RemainingBlendTimes.Empty();
+	}
+
+	// FAnimNode_Base interface
+	virtual void Initialize(const FAnimationInitializeContext& Context) override;
+	virtual void CacheBones(const FAnimationCacheBonesContext& Context) override;
+	virtual void Update(const FAnimationUpdateContext& Context) override;
+	virtual void Evaluate(FPoseContext& Output) override;
+	virtual void GatherDebugData(FNodeDebugData& DebugData) override;
+	// End of FAnimNode_Base interface
+
+#if WITH_EDITOR
+	virtual void AddPose()
+	{
+		BlendTime.Add(0.1f);
+		new (BlendPose) FPoseLink();
+	}
+
+	virtual void RemovePose(int32 PoseIndex)
+	{
+		BlendTime.RemoveAt(PoseIndex);
+		BlendPose.RemoveAt(PoseIndex);
+	}
+#endif
+
+protected:
+	virtual int32 GetActiveChildIndex() { return 0; }
+	virtual FString GetNodeName(FNodeDebugData& DebugData) { return DebugData.GetNodeName(this); }
+};
diff --git a/Engine/Source/Runtime/AnimGraphRuntime/Public/AnimNodes/AnimNode_BlendListByBool.h b/Engine/Source/Runtime/AnimGraphRuntime/Public/AnimNodes/AnimNode_BlendListByBool.h
new file mode 100644
--- /dev/null
+++ b/Engine/Source/Runtime/AnimGraphRuntime/Public/AnimNodes/AnimNode_BlendListByBool.h
@@ -0,0 +1,32 @@
+// Copyright 1998-2015 Epic Games, Inc. All Rights Reserved.
+
+#pragma once
+
+#include "Animation/AnimNodeBase.h"
+#include "AnimNodes/AnimNode_BlendListBase.h"
+#include "AnimNode_BlendListByBool.generated.h"
+
+// This node is effectively a 'branch', picking one of two input poses based on an input Boolean value
+USTRUCT()
+struct ANIMGRAPHRUNTIME_API FAnimNode_BlendListByBool : public FAnimNode_BlendListBase
+{
+	GENERATED_USTRUCT_BODY()
+public:
+	// Which input should be connected to the output?
+	UPROPERTY(EditAnywhere, BlueprintReadWrite, Category=Runtime, meta=(AlwaysAsPin))
+	bool bActiveValue;
+
+public:	
+	FAnimNode_BlendListByBool()
+		: FAnimNode_BlendListBase()
+	{
+#if	WITH_EDITOR
+		AddPose();
+		AddPose();
+#endif
+	}
+
+protected:
+	virtual int32 GetActiveChildIndex();
+	virtual FString GetNodeName(FNodeDebugData& DebugData) { return DebugData.GetNodeName(this); }
+};
diff --git a/Engine/Source/Runtime/AnimGraphRuntime/Public/AnimNodes/AnimNode_BlendListByEnum.h b/Engine/Source/Runtime/AnimGraphRuntime/Public/AnimNodes/AnimNode_BlendListByEnum.h
new file mode 100644
--- /dev/null
+++ b/Engine/Source/Runtime/AnimGraphRuntime/Public/AnimNodes/AnimNode_BlendListByEnum.h
@@ -0,0 +1,31 @@
+// Copyright 1998-2015 Epic Games, Inc. All Rights Reserved.
+
+#pragma once
+
+#include "AnimNode_BlendListBase.h"
+#include "AnimNode_BlendListByEnum.generated.h"
+
+// Blend List by Enum, it changes based on enum input that enters
+USTRUCT()
+struct ANIMGRAPHRUNTIME_API FAnimNode_BlendListByEnum : public FAnimNode_BlendListBase
+{
+	GENERATED_USTRUCT_BODY()
+public:
+	// Mapping from enum value to BlendPose index; there will be one entry per entry in the enum; entries out of range always map to pose index 0
+	UPROPERTY()
+	TArray<int32> EnumToPoseIndex;
+
+	// The currently selected pose (as an enum value)
+	UPROPERTY(EditAnywhere, BlueprintReadWrite, Category=Runtime, meta=(AlwaysAsPin))
+	mutable uint8 ActiveEnumValue;
+
+public:	
+	FAnimNode_BlendListByEnum()
+		: FAnimNode_BlendListBase()
+	{
+	}
+
+protected:
+	virtual int32 GetActiveChildIndex();
+	virtual FString GetNodeName(FNodeDebugData& DebugData) { return DebugData.GetNodeName(this); }
+};
diff --git a/Engine/Source/Runtime/AnimGraphRuntime/Public/AnimNodes/AnimNode_BlendListByInt.h b/Engine/Source/Runtime/AnimGraphRuntime/Public/AnimNodes/AnimNode_BlendListByInt.h
new file mode 100644
--- /dev/null
+++ b/Engine/Source/Runtime/AnimGraphRuntime/Public/AnimNodes/AnimNode_BlendListByInt.h
@@ -0,0 +1,29 @@
+// Copyright 1998-2015 Epic Games, Inc. All Rights Reserved.
+
+#pragma once
+
+#include "Animation/AnimNodeBase.h"
+#include "AnimNodes/AnimNode_BlendListBase.h"
+#include "AnimNode_BlendListByInt.generated.h"
+
+// Blend list node; has many children
+USTRUCT()
+struct ANIMGRAPHRUNTIME_API FAnimNode_BlendListByInt : public FAnimNode_BlendListBase
+{
+	GENERATED_USTRUCT_BODY()
+public:
+	UPROPERTY(EditAnywhere, BlueprintReadWrite, Category=Runtime, meta=(AlwaysAsPin))
+	mutable int32 ActiveChildIndex;
+
+
+public:	
+	FAnimNode_BlendListByInt()
+		:	FAnimNode_BlendListBase()
+		,	ActiveChildIndex(0)
+	{
+	}
+
+protected:
+	virtual int32 GetActiveChildIndex();
+	virtual FString GetNodeName(FNodeDebugData& DebugData) { return DebugData.GetNodeName(this); }
+};
diff --git a/Engine/Source/Runtime/AnimGraphRuntime/Public/AnimNodes/AnimNode_BlendSpaceEvaluator.h b/Engine/Source/Runtime/AnimGraphRuntime/Public/AnimNodes/AnimNode_BlendSpaceEvaluator.h
new file mode 100644
--- /dev/null
+++ b/Engine/Source/Runtime/AnimGraphRuntime/Public/AnimNodes/AnimNode_BlendSpaceEvaluator.h
@@ -0,0 +1,28 @@
+// Copyright 1998-2015 Epic Games, Inc. All Rights Reserved.
+
+#pragma once
+
+#include "AnimNode_BlendSpacePlayer.h"
+#include "AnimNode_BlendSpaceEvaluator.generated.h"
+
+// Evaluates a point in a blendspace, using a specific time input rather than advancing time internally.
+// Typically the playback position of the animation for this node will represent something other than time, like jump height.
+// This node will not trigger any notifies present in the associated sequence.
+USTRUCT()
+struct ANIMGRAPHRUNTIME_API FAnimNode_BlendSpaceEvaluator : public FAnimNode_BlendSpacePlayer
+{
+	GENERATED_USTRUCT_BODY()
+
+public:
+	/** Normalized time between [0,1]. The actual length of a blendspace is dynamic based on the coordinate, so it is exposed as a normalized value. */
+	UPROPERTY(EditAnywhere, BlueprintReadWrite, Category=Settings, meta=(PinShownByDefault))
+	mutable float NormalizedTime;
+
+public:	
+	FAnimNode_BlendSpaceEvaluator();
+
+	// FAnimNode_Base interface
+	virtual void UpdateAssetPlayer(const FAnimationUpdateContext& Context) override;
+	virtual void GatherDebugData(FNodeDebugData& DebugData) override;
+	// End of FAnimNode_Base interface
+};
diff --git a/Engine/Source/Runtime/AnimGraphRuntime/Public/AnimNodes/AnimNode_BlendSpacePlayer.h b/Engine/Source/Runtime/AnimGraphRuntime/Public/AnimNodes/AnimNode_BlendSpacePlayer.h
new file mode 100644
--- /dev/null
+++ b/Engine/Source/Runtime/AnimGraphRuntime/Public/AnimNodes/AnimNode_BlendSpacePlayer.h
@@ -0,0 +1,73 @@
+// Copyright 1998-2015 Epic Games, Inc. All Rights Reserved.
+
+#pragma once
+
+#include "Animation/AnimNodeBase.h"
+#include "Animation/AnimNode_AssetPlayerBase.h"
+#include "AnimNode_BlendSpacePlayer.generated.h"
+
+//@TODO: Comment
+USTRUCT()
+struct ANIMGRAPHRUNTIME_API FAnimNode_BlendSpacePlayer : public FAnimNode_AssetPlayerBase
+{
+	GENERATED_USTRUCT_BODY()
+
+public:
+	// The X coordinate to sample in the blendspace
+	UPROPERTY(EditAnywhere, BlueprintReadWrite, Category=Coordinates, meta=(PinShownByDefault))
+	mutable float X;
+
+	// The Y coordinate to sample in the blendspace
+	UPROPERTY(EditAnywhere, BlueprintReadWrite, Category = Coordinates, meta = (PinShownByDefault))
+	mutable float Y;
+
+	// The Z coordinate to sample in the blendspace
+	UPROPERTY(EditAnywhere, BlueprintReadWrite, Category = Coordinates, meta = (PinHiddenByDefault))
+	mutable float Z;
+
+	// The play rate multiplier. Can be negative, which will cause the animation to play in reverse.
+	UPROPERTY(EditAnywhere, BlueprintReadWrite, Category = Settings, meta = (DefaultValue = "1.0", PinHiddenByDefault))
+	mutable float PlayRate;
+
+	// Should the animation continue looping when it reaches the end?
+	UPROPERTY(EditAnywhere, BlueprintReadWrite, Category = Settings, meta = (DefaultValue = "true", PinHiddenByDefault))
+	mutable bool bLoop;
+
+	// The blendspace asset to play
+	UPROPERTY(EditAnywhere, BlueprintReadWrite, Category = Settings, meta = (PinHiddenByDefault))
+	UBlendSpaceBase* BlendSpace;
+
+	// The group index, assigned at compile time based on the editoronly GroupName (or INDEX_NONE if it is not part of any group)
+	UPROPERTY()
+	int32 GroupIndex;
+
+	// The role this player can assume within the group (ignored if GroupIndex is INDEX_NONE)
+	UPROPERTY()
+	TEnumAsByte<EAnimGroupRole::Type> GroupRole;
+protected:
+
+	UPROPERTY()
+	FBlendFilter BlendFilter;
+
+	UPROPERTY()
+	TArray<FBlendSampleData> BlendSampleDataCache;
+
+public:	
+	FAnimNode_BlendSpacePlayer();
+
+	// FAnimNode_Base interface
+	virtual void Initialize(const FAnimationInitializeContext& Context) override;
+	virtual void CacheBones(const FAnimationCacheBonesContext& Context) override;
+	virtual void UpdateAssetPlayer(const FAnimationUpdateContext& Context) override;
+	virtual void Evaluate(FPoseContext& Output) override;
+	virtual void OverrideAsset(UAnimationAsset* NewAsset) override;
+	virtual void GatherDebugData(FNodeDebugData& DebugData) override;
+	// End of FAnimNode_Base interface
+
+	float GetTimeFromEnd(float CurrentTime);
+
+	UAnimationAsset* GetAnimAsset();
+
+protected:
+	void UpdateInternal(const FAnimationUpdateContext& Context);
+};
diff --git a/Engine/Source/Runtime/AnimGraphRuntime/Public/AnimNodes/AnimNode_LayeredBoneBlend.h b/Engine/Source/Runtime/AnimGraphRuntime/Public/AnimNodes/AnimNode_LayeredBoneBlend.h
new file mode 100644
--- /dev/null
+++ b/Engine/Source/Runtime/AnimGraphRuntime/Public/AnimNodes/AnimNode_LayeredBoneBlend.h
@@ -0,0 +1,76 @@
+// Copyright 1998-2015 Epic Games, Inc. All Rights Reserved.
+
+#pragma once
+
+#include "Animation/AnimNodeBase.h"
+#include "Animation/AnimData/BoneMaskFilter.h"
+#include "Animation/AnimTypes.h"
+#include "AnimNode_LayeredBoneBlend.generated.h"
+
+
+// Layered blend (per bone); has dynamic number of blendposes that can blend per different bone sets
+USTRUCT()
+struct ANIMGRAPHRUNTIME_API FAnimNode_LayeredBoneBlend : public FAnimNode_Base
+{
+	GENERATED_USTRUCT_BODY()
+
+public:
+	UPROPERTY(EditAnywhere, BlueprintReadWrite, Category=Links)
+	FPoseLink BasePose;
+
+	//@TODO: Anim: Comment these members
+	UPROPERTY(EditAnywhere, BlueprintReadWrite, editfixedsize, Category=Links)
+	TArray<FPoseLink> BlendPoses;
+
+	UPROPERTY(EditAnywhere, BlueprintReadWrite, editfixedsize, Category=Config)
+	TArray<FInputBlendPose> LayerSetup;
+
+	UPROPERTY(EditAnywhere, BlueprintReadWrite, editfixedsize, Category=Runtime, meta=(PinShownByDefault))
+	TArray<float> BlendWeights;
+
+	UPROPERTY(EditAnywhere, BlueprintReadWrite, Category=Config)
+	bool bMeshSpaceRotationBlend;
+	
+	UPROPERTY(EditAnywhere, BlueprintReadWrite, Category=Config)
+	TEnumAsByte<enum ECurveBlendOption::Type>	CurveBlendOption;
+
+protected:
+	TArray<FPerBoneBlendWeight> DesiredBoneBlendWeights;
+	TArray<FPerBoneBlendWeight> CurrentBoneBlendWeights;
+
+public:	
+	FAnimNode_LayeredBoneBlend()
+	{
+	}
+
+	// FAnimNode_Base interface
+	virtual void Initialize(const FAnimationInitializeContext& Context) override;
+	virtual void CacheBones(const FAnimationCacheBonesContext& Context) override;
+	virtual void Update(const FAnimationUpdateContext& Context) override;
+	virtual void Evaluate(FPoseContext& Output) override;
+	virtual void GatherDebugData(FNodeDebugData& DebugData) override;
+	// End of FAnimNode_Base interface
+
+#if WITH_EDITOR
+	void AddPose()
+	{
+		BlendWeights.Add(1.f);
+		new (BlendPoses) FPoseLink();
+		new (LayerSetup) FInputBlendPose();
+	}
+
+	void RemovePose(int32 PoseIndex)
+	{
+		BlendWeights.RemoveAt(PoseIndex);
+		BlendPoses.RemoveAt(PoseIndex);
+		LayerSetup.RemoveAt(PoseIndex);
+	}
+
+	// ideally you don't like to get to situation where it becomes inconsistent, but this happened, 
+	// and we don't know what caused this. Possibly copy/paste, but I tried copy/paste and that didn't work
+	// so here we add code to fix this up manually in editor, so that they can continue working on it. 
+	void ValidateData();
+#endif
+
+	void ReinitializeBoneBlendWeights(const FBoneContainer& RequiredBones, const USkeleton* Skeleton);
+};
diff --git a/Engine/Source/Runtime/AnimGraphRuntime/Public/AnimNodes/AnimNode_RefPose.h b/Engine/Source/Runtime/AnimGraphRuntime/Public/AnimNodes/AnimNode_RefPose.h
new file mode 100644
--- /dev/null
+++ b/Engine/Source/Runtime/AnimGraphRuntime/Public/AnimNodes/AnimNode_RefPose.h
@@ -0,0 +1,46 @@
+// Copyright 1998-2015 Epic Games, Inc. All Rights Reserved.
+
+#pragma once
+#include "Animation/AnimNodeBase.h"
+#include "AnimNode_RefPose.generated.h"
+
+UENUM()
+enum ERefPoseType
+{
+	EIT_LocalSpace, 
+	EIT_Additive
+};
+
+// RefPose pose nodes - ref pose or additive RefPose pose
+USTRUCT()
+struct ANIMGRAPHRUNTIME_API FAnimNode_RefPose : public FAnimNode_Base
+{
+	GENERATED_USTRUCT_BODY()
+
+	UPROPERTY()
+	TEnumAsByte<ERefPoseType> RefPoseType;
+	
+public:	
+	FAnimNode_RefPose()
+		:	FAnimNode_Base()
+		, 	RefPoseType(EIT_LocalSpace)
+	{
+	}
+
+	virtual void Evaluate(FPoseContext& Output) override;
+	virtual void GatherDebugData(FNodeDebugData& DebugData) override;
+
+};
+
+USTRUCT()
+struct ANIMGRAPHRUNTIME_API FAnimNode_MeshSpaceRefPose : public FAnimNode_Base
+{
+	GENERATED_USTRUCT_BODY()
+public:	
+	FAnimNode_MeshSpaceRefPose()
+		:	FAnimNode_Base()
+	{
+	}
+
+	virtual void EvaluateComponentSpace(FComponentSpacePoseContext& Output);
+};
diff --git a/Engine/Source/Runtime/AnimGraphRuntime/Public/AnimNodes/AnimNode_Root.h b/Engine/Source/Runtime/AnimGraphRuntime/Public/AnimNodes/AnimNode_Root.h
new file mode 100644
--- /dev/null
+++ b/Engine/Source/Runtime/AnimGraphRuntime/Public/AnimNodes/AnimNode_Root.h
@@ -0,0 +1,26 @@
+// Copyright 1998-2015 Epic Games, Inc. All Rights Reserved.
+
+#pragma once
+#include "Animation/AnimNodeBase.h"
+#include "AnimNode_Root.generated.h"
+
+// Root node of an animation tree (sink)
+USTRUCT()
+struct ANIMGRAPHRUNTIME_API FAnimNode_Root : public FAnimNode_Base
+{
+	GENERATED_USTRUCT_BODY()
+
+	UPROPERTY(EditAnywhere, BlueprintReadWrite, Category=Links)
+	FPoseLink Result;
+
+public:	
+	FAnimNode_Root();
+
+	// FAnimNode_Base interface
+	virtual void Initialize(const FAnimationInitializeContext& Context) override;
+	virtual void CacheBones(const FAnimationCacheBonesContext& Context) override;
+	virtual void Update(const FAnimationUpdateContext& Context) override;
+	virtual void Evaluate(FPoseContext& Output) override;
+	virtual void GatherDebugData(FNodeDebugData& DebugData) override;
+	// End of FAnimNode_Base interface
+};
diff --git a/Engine/Source/Runtime/AnimGraphRuntime/Public/AnimNodes/AnimNode_RotateRootBone.h b/Engine/Source/Runtime/AnimGraphRuntime/Public/AnimNodes/AnimNode_RotateRootBone.h
new file mode 100644
--- /dev/null
+++ b/Engine/Source/Runtime/AnimGraphRuntime/Public/AnimNodes/AnimNode_RotateRootBone.h
@@ -0,0 +1,36 @@
+// Copyright 1998-2015 Epic Games, Inc. All Rights Reserved.
+
+#pragma once
+
+#include "Animation/AnimNodeBase.h"
+#include "AnimNode_RotateRootBone.generated.h"
+
+//@TODO: Comment
+USTRUCT()
+struct ANIMGRAPHRUNTIME_API FAnimNode_RotateRootBone : public FAnimNode_Base
+{
+	GENERATED_USTRUCT_BODY()
+
+	UPROPERTY(EditAnywhere, BlueprintReadWrite, Category=Links)
+	FPoseLink BasePose;
+
+	UPROPERTY(EditAnywhere, BlueprintReadWrite, Category=Settings, meta=(PinShownByDefault))
+	mutable float Pitch;
+
+	UPROPERTY(EditAnywhere, BlueprintReadWrite, Category=Settings, meta=(PinShownByDefault))
+	mutable float Yaw;
+
+	UPROPERTY(EditAnywhere, BlueprintReadWrite, Category=Settings, meta=(PinShownByDefault))
+	mutable FRotator MeshToComponent;
+
+public:	
+	FAnimNode_RotateRootBone();
+
+	// FAnimNode_Base interface
+	virtual void Initialize(const FAnimationInitializeContext& Context) override;
+	virtual void CacheBones(const FAnimationCacheBonesContext& Context) override;
+	virtual void Update(const FAnimationUpdateContext& Context) override;
+	virtual void Evaluate(FPoseContext& Output) override;
+	virtual void GatherDebugData(FNodeDebugData& DebugData) override;
+	// End of FAnimNode_Base interface
+};
diff --git a/Engine/Source/Runtime/AnimGraphRuntime/Public/AnimNodes/AnimNode_RotationOffsetBlendSpace.h b/Engine/Source/Runtime/AnimGraphRuntime/Public/AnimNodes/AnimNode_RotationOffsetBlendSpace.h
new file mode 100644
--- /dev/null
+++ b/Engine/Source/Runtime/AnimGraphRuntime/Public/AnimNodes/AnimNode_RotationOffsetBlendSpace.h
@@ -0,0 +1,27 @@
+// Copyright 1998-2015 Epic Games, Inc. All Rights Reserved.
+
+#pragma once
+
+#include "AnimNode_BlendSpacePlayer.h"
+#include "AnimNode_RotationOffsetBlendSpace.generated.h"
+
+//@TODO: Comment
+USTRUCT()
+struct ANIMGRAPHRUNTIME_API FAnimNode_RotationOffsetBlendSpace : public FAnimNode_BlendSpacePlayer
+{
+	GENERATED_USTRUCT_BODY()
+
+	UPROPERTY(EditAnywhere, BlueprintReadWrite, Category=Links)
+	FPoseLink BasePose;
+
+public:	
+	FAnimNode_RotationOffsetBlendSpace();
+
+	// FAnimNode_Base interface
+	virtual void Initialize(const FAnimationInitializeContext& Context) override;
+	virtual void CacheBones(const FAnimationCacheBonesContext& Context) override;
+	virtual void UpdateAssetPlayer(const FAnimationUpdateContext& Context) override;
+	virtual void Evaluate(FPoseContext& Output) override;
+	virtual void GatherDebugData(FNodeDebugData& DebugData) override;
+	// End of FAnimNode_Base interface
+};
diff --git a/Engine/Source/Runtime/AnimGraphRuntime/Public/AnimNodes/AnimNode_SaveCachedPose.h b/Engine/Source/Runtime/AnimGraphRuntime/Public/AnimNodes/AnimNode_SaveCachedPose.h
new file mode 100644
--- /dev/null
+++ b/Engine/Source/Runtime/AnimGraphRuntime/Public/AnimNodes/AnimNode_SaveCachedPose.h
@@ -0,0 +1,41 @@
+// Copyright 1998-2015 Epic Games, Inc. All Rights Reserved.
+
+#pragma once
+#include "Animation/AnimNodeBase.h"
+#include "AnimNode_SaveCachedPose.generated.h"
+
+USTRUCT()
+struct ANIMGRAPHRUNTIME_API FAnimNode_SaveCachedPose : public FAnimNode_Base
+{
+	GENERATED_USTRUCT_BODY()
+
+	UPROPERTY(transient)
+	int16 LastInitializedContextCounter;
+
+	UPROPERTY(transient)
+	int16 LastCacheBonesContextCounter;
+
+	UPROPERTY(transient)
+	int16 LastUpdatedContextCounter;
+
+	UPROPERTY(transient)
+	int16 LastEvaluatedContextCounter;
+
+	UPROPERTY(EditAnywhere, BlueprintReadWrite, Category=Links)
+	FPoseLink Pose;
+
+protected:
+	FCompactPose CachedPose;
+	FBlendedCurve CachedCurve;
+
+public:	
+	FAnimNode_SaveCachedPose();
+
+	// FAnimNode_Base interface
+	virtual void Initialize(const FAnimationInitializeContext& Context) override;
+	virtual void CacheBones(const FAnimationCacheBonesContext& Context) override;
+	virtual void Update(const FAnimationUpdateContext& Context) override;
+	virtual void Evaluate(FPoseContext& Output) override;
+	virtual void GatherDebugData(FNodeDebugData& DebugData) override;
+	// End of FAnimNode_Base interface
+};
diff --git a/Engine/Source/Runtime/AnimGraphRuntime/Public/AnimNodes/AnimNode_SequenceEvaluator.h b/Engine/Source/Runtime/AnimGraphRuntime/Public/AnimNodes/AnimNode_SequenceEvaluator.h
new file mode 100644
--- /dev/null
+++ b/Engine/Source/Runtime/AnimGraphRuntime/Public/AnimNodes/AnimNode_SequenceEvaluator.h
@@ -0,0 +1,44 @@
+// Copyright 1998-2015 Epic Games, Inc. All Rights Reserved.
+
+#pragma once
+#include "Animation/AnimNodeBase.h"
+#include "Animation/AnimNode_AssetPlayerBase.h"
+#include "AnimNode_SequenceEvaluator.generated.h"
+
+// Evaluates a point in an anim sequence, using a specific time input rather than advancing time internally.
+// Typically the playback position of the animation for this node will represent something other than time, like jump height.
+// This node will not trigger any notifies present in the associated sequence.
+USTRUCT()
+struct ANIMGRAPHRUNTIME_API FAnimNode_SequenceEvaluator : public FAnimNode_AssetPlayerBase
+{
+	GENERATED_USTRUCT_BODY()
+public:
+	// The animation sequence asset to evaluate
+	UPROPERTY(EditAnywhere, BlueprintReadWrite, Category = Settings, meta = (PinHiddenByDefault))
+	UAnimSequenceBase* Sequence;
+
+	// The time at which to evaluate the associated sequence
+	UPROPERTY(EditAnywhere, BlueprintReadWrite, Category=Settings, meta=(PinShownByDefault))
+	mutable float ExplicitTime;
+
+public:	
+	FAnimNode_SequenceEvaluator()
+		: ExplicitTime(0.0f)
+	{
+	}
+
+	// FAnimNode_Base interface
+	virtual void Initialize(const FAnimationInitializeContext& Context) override;
+	virtual void CacheBones(const FAnimationCacheBonesContext& Context) override;
+	virtual void UpdateAssetPlayer(const FAnimationUpdateContext& Context) override;
+	virtual void Evaluate(FPoseContext& Output) override;
+	virtual void OverrideAsset(UAnimationAsset* NewAsset) override;
+	virtual void GatherDebugData(FNodeDebugData& DebugData) override;
+	// End of FAnimNode_Base interface
+
+	// FAnimNode_AssetPlayerBase Interface
+	virtual float GetAccumulatedTime() {return ExplicitTime;}
+	virtual void SetAccumulatedTime(const float& NewTime) {ExplicitTime = NewTime;}
+	virtual UAnimationAsset* GetAnimAsset() {return Sequence;}
+	// End of FAnimNode_AssetPlayerBase Interface
+};
diff --git a/Engine/Source/Runtime/AnimGraphRuntime/Public/AnimNodes/AnimNode_Slot.h b/Engine/Source/Runtime/AnimGraphRuntime/Public/AnimNodes/AnimNode_Slot.h
new file mode 100644
--- /dev/null
+++ b/Engine/Source/Runtime/AnimGraphRuntime/Public/AnimNodes/AnimNode_Slot.h
@@ -0,0 +1,46 @@
+// Copyright 1998-2015 Epic Games, Inc. All Rights Reserved.
+
+#pragma once
+#include "Animation/AnimNodeBase.h"
+#include "AnimNode_Slot.generated.h"
+
+// An animation slot node normally acts as a passthru, but a montage or PlaySlotAnimation call from
+// game code can cause an animation to blend in and be played on the slot temporarily, overriding the
+// Source input.
+USTRUCT()
+struct ANIMGRAPHRUNTIME_API FAnimNode_Slot : public FAnimNode_Base
+{
+	GENERATED_USTRUCT_BODY()
+
+	// The source input, passed thru to the output unless a montage or slot animation is currently playing
+	UPROPERTY(EditAnywhere, BlueprintReadWrite, Category=Links)
+	FPoseLink Source;
+
+	// The name of this slot, exposed to gameplay code, etc...
+	UPROPERTY(EditAnywhere, BlueprintReadWrite, Category=Settings, meta=(CustomizeProperty))
+	FName SlotName;
+
+protected:
+	/** Weight of Source Branch. This is the weight of the input pose coming from children.
+	This is different than (1.f - SourceWeight) since the Slot can play additive animations, which are overlayed on top of the source pose. */
+	float SourceWeight;
+
+	/** Weight of Slot Node. Determined by Montages weight playing on this slot */
+	float SlotNodeWeight;
+
+private:
+	/** Counter to register slot node just once per initialization */
+	UPROPERTY(transient)
+	int16 LastSlotNodeInitializationCounter;
+
+public:	
+	FAnimNode_Slot();
+
+	// FAnimNode_Base interface
+	virtual void Initialize(const FAnimationInitializeContext& Context) override;
+	virtual void CacheBones(const FAnimationCacheBonesContext& Context) override;
+	virtual void Update(const FAnimationUpdateContext& Context) override;
+	virtual void Evaluate(FPoseContext& Output) override;
+	virtual void GatherDebugData(FNodeDebugData& DebugData) override;
+	// End of FAnimNode_Base interface
+};
diff --git a/Engine/Source/Runtime/AnimGraphRuntime/Public/AnimNodes/AnimNode_TwoWayBlend.h b/Engine/Source/Runtime/AnimGraphRuntime/Public/AnimNodes/AnimNode_TwoWayBlend.h
new file mode 100644
--- /dev/null
+++ b/Engine/Source/Runtime/AnimGraphRuntime/Public/AnimNodes/AnimNode_TwoWayBlend.h
@@ -0,0 +1,42 @@
+// Copyright 1998-2015 Epic Games, Inc. All Rights Reserved.
+
+#pragma once
+
+#include "Animation/AnimNodeBase.h"
+#include "Animation/InputScaleBias.h"
+#include "AnimNode_TwoWayBlend.generated.h"
+
+// This represents a baked transition
+USTRUCT()
+struct ANIMGRAPHRUNTIME_API FAnimationNode_TwoWayBlend : public FAnimNode_Base
+{
+	GENERATED_USTRUCT_BODY()
+public:
+
+	UPROPERTY(EditAnywhere, BlueprintReadWrite, Category=Links)
+	FPoseLink A;
+
+	UPROPERTY(EditAnywhere, BlueprintReadWrite, Category=Links)
+	FPoseLink B;
+
+	UPROPERTY(EditAnywhere, BlueprintReadWrite, Category=Settings, meta=(PinShownByDefault))
+	mutable float Alpha;
+
+	UPROPERTY(EditAnywhere, BlueprintReadWrite, Category=Settings)
+	FInputScaleBias AlphaScaleBias;
+
+public:
+	FAnimationNode_TwoWayBlend()
+		: Alpha(0.0f)
+	{
+	}
+
+	// FAnimNode_Base interface
+	virtual void Initialize(const FAnimationInitializeContext& Context) override;
+	virtual void CacheBones(const FAnimationCacheBonesContext& Context) override;
+	virtual void Update(const FAnimationUpdateContext& Context) override;
+	virtual void Evaluate(FPoseContext& Output) override;
+	virtual void GatherDebugData(FNodeDebugData& DebugData) override;
+	// End of FAnimNode_Base interface
+};
+
diff --git a/Engine/Source/Runtime/Engine/Classes/Animation/AnimCompositeBase.h b/Engine/Source/Runtime/Engine/Classes/Animation/AnimCompositeBase.h
--- a/Engine/Source/Runtime/Engine/Classes/Animation/AnimCompositeBase.h
+++ b/Engine/Source/Runtime/Engine/Classes/Animation/AnimCompositeBase.h
@@ -133,7 +133,7 @@ struct FAnimSegment
 	/**
 	 * Get Animation Data, for now have weight to be controlled here, in the future, it will be controlled in Track
 	 */
-	UAnimSequenceBase* GetAnimationData(float PositionInTrack, float& PositionInAnim, float& Weight) const;
+	ENGINE_API UAnimSequenceBase* GetAnimationData(float PositionInTrack, float& PositionInAnim, float& Weight) const;
 
 	/** Converts 'Track Position' to position on AnimSequence.
 	 * Note: doesn't check that position is in valid range, must do that before calling this function! */
diff --git a/Engine/Source/Runtime/Engine/Classes/Animation/AnimMontage.h b/Engine/Source/Runtime/Engine/Classes/Animation/AnimMontage.h
--- a/Engine/Source/Runtime/Engine/Classes/Animation/AnimMontage.h
+++ b/Engine/Source/Runtime/Engine/Classes/Animation/AnimMontage.h
@@ -412,7 +412,7 @@ class UAnimMontage : public UAnimCompositeBase
 #endif // WITH_EDITORONLY_DATA
 
 	/** return true if valid slot */
-	bool IsValidSlot(FName InSlotName) const;
+	ENGINE_API bool IsValidSlot(FName InSlotName) const;
 
 public:
 	// Begin UObject Interface
@@ -504,7 +504,7 @@ public:
 	float CalculatePos(FCompositeSection &Section, float PosWithinCompositeSection) const;
 	
 	/** Prototype function to get animation data - this will need rework */
-	const FAnimTrack* GetAnimationData(FName SlotName) const;
+	ENGINE_API const FAnimTrack* GetAnimationData(FName SlotName) const;
 
 	/** Returns whether the anim sequences this montage have root motion enabled */
 	virtual bool HasRootMotion() const override;
diff --git a/Engine/Source/Runtime/Engine/Classes/Animation/AnimNode_ApplyAdditive.h b/Engine/Source/Runtime/Engine/Classes/Animation/AnimNode_ApplyAdditive.h
deleted file mode 100644
diff --git a/Engine/Source/Runtime/Engine/Classes/Animation/AnimNode_BlendListBase.h b/Engine/Source/Runtime/Engine/Classes/Animation/AnimNode_BlendListBase.h
deleted file mode 100644
diff --git a/Engine/Source/Runtime/Engine/Classes/Animation/AnimNode_BlendListByBool.h b/Engine/Source/Runtime/Engine/Classes/Animation/AnimNode_BlendListByBool.h
deleted file mode 100644
diff --git a/Engine/Source/Runtime/Engine/Classes/Animation/AnimNode_BlendListByEnum.h b/Engine/Source/Runtime/Engine/Classes/Animation/AnimNode_BlendListByEnum.h
deleted file mode 100644
diff --git a/Engine/Source/Runtime/Engine/Classes/Animation/AnimNode_BlendListByInt.h b/Engine/Source/Runtime/Engine/Classes/Animation/AnimNode_BlendListByInt.h
deleted file mode 100644
diff --git a/Engine/Source/Runtime/Engine/Classes/Animation/AnimNode_BlendSpaceEvaluator.h b/Engine/Source/Runtime/Engine/Classes/Animation/AnimNode_BlendSpaceEvaluator.h
deleted file mode 100644
diff --git a/Engine/Source/Runtime/Engine/Classes/Animation/AnimNode_BlendSpacePlayer.h b/Engine/Source/Runtime/Engine/Classes/Animation/AnimNode_BlendSpacePlayer.h
deleted file mode 100644
diff --git a/Engine/Source/Runtime/Engine/Classes/Animation/AnimNode_LayeredBoneBlend.h b/Engine/Source/Runtime/Engine/Classes/Animation/AnimNode_LayeredBoneBlend.h
deleted file mode 100644
diff --git a/Engine/Source/Runtime/Engine/Classes/Animation/AnimNode_RefPose.h b/Engine/Source/Runtime/Engine/Classes/Animation/AnimNode_RefPose.h
deleted file mode 100644
diff --git a/Engine/Source/Runtime/Engine/Classes/Animation/AnimNode_Root.h b/Engine/Source/Runtime/Engine/Classes/Animation/AnimNode_Root.h
deleted file mode 100644
diff --git a/Engine/Source/Runtime/Engine/Classes/Animation/AnimNode_RotateRootBone.h b/Engine/Source/Runtime/Engine/Classes/Animation/AnimNode_RotateRootBone.h
deleted file mode 100644
diff --git a/Engine/Source/Runtime/Engine/Classes/Animation/AnimNode_RotationOffsetBlendSpace.h b/Engine/Source/Runtime/Engine/Classes/Animation/AnimNode_RotationOffsetBlendSpace.h
deleted file mode 100644
diff --git a/Engine/Source/Runtime/Engine/Classes/Animation/AnimNode_SaveCachedPose.h b/Engine/Source/Runtime/Engine/Classes/Animation/AnimNode_SaveCachedPose.h
deleted file mode 100644
diff --git a/Engine/Source/Runtime/Engine/Classes/Animation/AnimNode_SequenceEvaluator.h b/Engine/Source/Runtime/Engine/Classes/Animation/AnimNode_SequenceEvaluator.h
deleted file mode 100644
diff --git a/Engine/Source/Runtime/Engine/Classes/Animation/AnimNode_Slot.h b/Engine/Source/Runtime/Engine/Classes/Animation/AnimNode_Slot.h
deleted file mode 100644
diff --git a/Engine/Source/Runtime/Engine/Classes/Animation/AnimNode_TwoWayBlend.h b/Engine/Source/Runtime/Engine/Classes/Animation/AnimNode_TwoWayBlend.h
deleted file mode 100644
diff --git a/Engine/Source/Runtime/Engine/Classes/Animation/BlendSpaceBase.h b/Engine/Source/Runtime/Engine/Classes/Animation/BlendSpaceBase.h
--- a/Engine/Source/Runtime/Engine/Classes/Animation/BlendSpaceBase.h
+++ b/Engine/Source/Runtime/Engine/Classes/Animation/BlendSpaceBase.h
@@ -277,7 +277,7 @@ public:
 	/**
 	 * BlendSpace Get Animation Pose function
 	 */
-	void GetAnimationPose(TArray<FBlendSampleData>& BlendSampleDataCache, /*out*/ FCompactPose& OutPose, /*out*/ FBlendedCurve& OutCurve);
+	ENGINE_API void GetAnimationPose(TArray<FBlendSampleData>& BlendSampleDataCache, /*out*/ FCompactPose& OutPose, /*out*/ FBlendedCurve& OutCurve);
 
 	/** Accessor for blend parameter **/
 	ENGINE_API const FBlendParameter& GetBlendParameter(int32 Index)
@@ -350,7 +350,7 @@ public:
 	ENGINE_API bool IsTooCloseToExistingSamplePoint(const FVector& SampleValue, int32 OriginalIndex) const;
 
 	/** Initialize BlendSpace for runtime. It needs certain data to be reinitialized per instsance **/
-	void InitializeFilter(FBlendFilter* Filter) const;
+	ENGINE_API void InitializeFilter(FBlendFilter* Filter) const;
 
 	/** 
 	 * Get PerBoneInterpolationIndex for the input BoneIndex
diff --git a/Engine/Source/Runtime/Engine/Private/Animation/AnimInstance.cpp b/Engine/Source/Runtime/Engine/Private/Animation/AnimInstance.cpp
--- a/Engine/Source/Runtime/Engine/Private/Animation/AnimInstance.cpp
+++ b/Engine/Source/Runtime/Engine/Private/Animation/AnimInstance.cpp
@@ -8,7 +8,7 @@
 #include "AnimationRuntime.h"
 #include "AnimationUtils.h"
 #include "Animation/AnimStats.h"
-#include "Animation/AnimNode_StateMachine.h"
+#include "Animation/AnimBlueprintGeneratedClass.h"
 #include "GameFramework/Character.h"
 #include "ParticleDefinitions.h"
 #include "DisplayDebugHelpers.h"
@@ -18,12 +18,11 @@
 #include "Animation/AnimNotifies/AnimNotifyState.h"
 #include "Animation/AnimNotifies/AnimNotify.h"
 #include "Animation/AnimMontage.h"
-#include "Animation/AnimInstance.h"
+#include "Animation/AnimNodeBase.h"
+#include "Animation/AnimNode_StateMachine.h"
 #include "Animation/AnimNode_TransitionResult.h"
-#include "Animation/AnimNode_SequencePlayer.h"
-#include "Animation/AnimNode_BlendSpacePlayer.h"
 #include "Animation/AnimNode_AssetPlayerBase.h"
-
+#include "Animation/AnimInstance.h"
 /** Anim stats */
 
 DEFINE_STAT(STAT_UpdateSkelMeshBounds);
diff --git a/Engine/Source/Runtime/Engine/Private/Animation/AnimNode_ApplyAdditive.cpp b/Engine/Source/Runtime/Engine/Private/Animation/AnimNode_ApplyAdditive.cpp
deleted file mode 100644
diff --git a/Engine/Source/Runtime/Engine/Private/Animation/AnimNode_BlendListBase.cpp b/Engine/Source/Runtime/Engine/Private/Animation/AnimNode_BlendListBase.cpp
deleted file mode 100644
diff --git a/Engine/Source/Runtime/Engine/Private/Animation/AnimNode_BlendListByBool.cpp b/Engine/Source/Runtime/Engine/Private/Animation/AnimNode_BlendListByBool.cpp
deleted file mode 100644
diff --git a/Engine/Source/Runtime/Engine/Private/Animation/AnimNode_BlendListByEnum.cpp b/Engine/Source/Runtime/Engine/Private/Animation/AnimNode_BlendListByEnum.cpp
deleted file mode 100644
diff --git a/Engine/Source/Runtime/Engine/Private/Animation/AnimNode_BlendListByInt.cpp b/Engine/Source/Runtime/Engine/Private/Animation/AnimNode_BlendListByInt.cpp
deleted file mode 100644
diff --git a/Engine/Source/Runtime/Engine/Private/Animation/AnimNode_BlendSpaceEvaluator.cpp b/Engine/Source/Runtime/Engine/Private/Animation/AnimNode_BlendSpaceEvaluator.cpp
deleted file mode 100644
diff --git a/Engine/Source/Runtime/Engine/Private/Animation/AnimNode_BlendSpacePlayer.cpp b/Engine/Source/Runtime/Engine/Private/Animation/AnimNode_BlendSpacePlayer.cpp
deleted file mode 100644
diff --git a/Engine/Source/Runtime/Engine/Private/Animation/AnimNode_LayeredBoneBlend.cpp b/Engine/Source/Runtime/Engine/Private/Animation/AnimNode_LayeredBoneBlend.cpp
deleted file mode 100644
diff --git a/Engine/Source/Runtime/Engine/Private/Animation/AnimNode_RefPose.cpp b/Engine/Source/Runtime/Engine/Private/Animation/AnimNode_RefPose.cpp
deleted file mode 100644
diff --git a/Engine/Source/Runtime/Engine/Private/Animation/AnimNode_Root.cpp b/Engine/Source/Runtime/Engine/Private/Animation/AnimNode_Root.cpp
deleted file mode 100644
diff --git a/Engine/Source/Runtime/Engine/Private/Animation/AnimNode_RotateRootBone.cpp b/Engine/Source/Runtime/Engine/Private/Animation/AnimNode_RotateRootBone.cpp
deleted file mode 100644
diff --git a/Engine/Source/Runtime/Engine/Private/Animation/AnimNode_RotationOffsetBlendSpace.cpp b/Engine/Source/Runtime/Engine/Private/Animation/AnimNode_RotationOffsetBlendSpace.cpp
deleted file mode 100644
diff --git a/Engine/Source/Runtime/Engine/Private/Animation/AnimNode_SaveCachedPose.cpp b/Engine/Source/Runtime/Engine/Private/Animation/AnimNode_SaveCachedPose.cpp
deleted file mode 100644
diff --git a/Engine/Source/Runtime/Engine/Private/Animation/AnimNode_SequenceEvaluator.cpp b/Engine/Source/Runtime/Engine/Private/Animation/AnimNode_SequenceEvaluator.cpp
deleted file mode 100644
diff --git a/Engine/Source/Runtime/Engine/Private/Animation/AnimNode_Slot.cpp b/Engine/Source/Runtime/Engine/Private/Animation/AnimNode_Slot.cpp
deleted file mode 100644
diff --git a/Engine/Source/Runtime/Engine/Private/Animation/AnimNode_TwoWayBlend.cpp b/Engine/Source/Runtime/Engine/Private/Animation/AnimNode_TwoWayBlend.cpp
deleted file mode 100644
diff --git a/Engine/Source/Runtime/Engine/Public/AlphaBlend.h b/Engine/Source/Runtime/Engine/Public/AlphaBlend.h
--- a/Engine/Source/Runtime/Engine/Public/AlphaBlend.h
+++ b/Engine/Source/Runtime/Engine/Public/AlphaBlend.h
@@ -45,7 +45,7 @@ enum class EAlphaBlendOption : uint8
  * Alpha Blend Type
  */
 USTRUCT()
-struct FAlphaBlend
+struct ENGINE_API FAlphaBlend
 {
 	GENERATED_BODY()
 
diff --git a/Engine/Source/Runtime/Engine/Public/Animation/AnimStats.h b/Engine/Source/Runtime/Engine/Public/Animation/AnimStats.h
--- a/Engine/Source/Runtime/Engine/Public/Animation/AnimStats.h
+++ b/Engine/Source/Runtime/Engine/Public/Animation/AnimStats.h
@@ -31,11 +31,11 @@ DECLARE_CYCLE_STAT_EXTERN(TEXT("BlueprintUpdateAnimation"), STAT_BlueprintUpdate
 DECLARE_CYCLE_STAT_EXTERN(TEXT("NativeUpdateAnimation"), STAT_NativeUpdateAnimation, STATGROUP_Anim, );
 DECLARE_CYCLE_STAT_EXTERN(TEXT("BlueprintPostEvaluateAnimation"), STAT_BlueprintPostEvaluateAnimation, STATGROUP_Anim, );
 
-#define DO_ANIMSTAT_PROCESSING(StatName) DECLARE_CYCLE_STAT_EXTERN(TEXT(#StatName), STAT_ ## StatName, STATGROUP_Anim, )
+#define DO_ANIMSTAT_PROCESSING(StatName) DECLARE_CYCLE_STAT_EXTERN(TEXT(#StatName), STAT_ ## StatName, STATGROUP_Anim, ENGINE_API)
 #include "AnimMTStats.h"
 #undef DO_ANIMSTAT_PROCESSING
 
-#define DO_ANIMSTAT_PROCESSING(StatName) DECLARE_CYCLE_STAT_EXTERN(TEXT(#StatName) TEXT("_WorkerThread"), STAT_ ## StatName ## _WorkerThread, STATGROUP_Anim, )
+#define DO_ANIMSTAT_PROCESSING(StatName) DECLARE_CYCLE_STAT_EXTERN(TEXT(#StatName) TEXT("_WorkerThread"), STAT_ ## StatName ## _WorkerThread, STATGROUP_Anim, ENGINE_API)
 #include "AnimMTStats.h"
 #undef DO_ANIMSTAT_PROCESSING
 
diff --git a/Engine/Source/Runtime/Engine/Public/BonePose.h b/Engine/Source/Runtime/Engine/Public/BonePose.h
--- a/Engine/Source/Runtime/Engine/Public/BonePose.h
+++ b/Engine/Source/Runtime/Engine/Public/BonePose.h
@@ -154,7 +154,7 @@ struct FCompactPoseBoneIndexReverseIterator
 	FCompactPoseBoneIndex operator*() const { return FCompactPoseBoneIndex(Index); }
 };
 
-struct FCompactPose : FBasePose<FCompactPoseBoneIndex>
+struct ENGINE_API FCompactPose : FBasePose<FCompactPoseBoneIndex>
 {
 public:
 	typedef FCompactPoseBoneIndex BoneIndexType;
@@ -244,7 +244,7 @@ public:
 	}
 
 	// Sets this pose to the supplied BoneContainers ref pose
-	ENGINE_API void ResetToRefPose(const FBoneContainer& RequiredBones);
+	void ResetToRefPose(const FBoneContainer& RequiredBones);
 
 	// Sets every bone transform to Identity
 	void ResetToIdentity();
