Renamed variable and deleted unused variable.

diff --git a/Engine/Source/Runtime/AnimGraphRuntime/Private/AnimNodes/AnimNode_CopyPoseFromMesh.cpp b/Engine/Source/Runtime/AnimGraphRuntime/Private/AnimNodes/AnimNode_CopyPoseFromMesh.cpp
--- a/Engine/Source/Runtime/AnimGraphRuntime/Private/AnimNodes/AnimNode_CopyPoseFromMesh.cpp
+++ b/Engine/Source/Runtime/AnimGraphRuntime/Private/AnimNodes/AnimNode_CopyPoseFromMesh.cpp
@@ -23,11 +23,11 @@ void FAnimNode_CopyPoseFromMesh::Update(const FAnimationUpdateContext& Context)
 {
 	EvaluateGraphExposedInputs.Execute(Context);
 
-	if (CurrentlyUsedMeshComponent.IsValid() && CurrentlyUsedMeshComponent.Get() != SourceMeshComponent)
+	if (CurrentlyUsedSourceMeshComponent.IsValid() && CurrentlyUsedSourceMeshComponent.Get() != SourceMeshComponent)
 	{
 		ReinitializeMeshComponent(Context.AnimInstance);
 	}
-	else if (!CurrentlyUsedMeshComponent.IsValid() && SourceMeshComponent)
+	else if (!CurrentlyUsedSourceMeshComponent.IsValid() && SourceMeshComponent)
 	{
 		ReinitializeMeshComponent(Context.AnimInstance);
 	}
@@ -38,7 +38,7 @@ void FAnimNode_CopyPoseFromMesh::Evaluate(FPoseContext& Output)
 	FCompactPose& OutPose = Output.Pose;
 	OutPose.ResetToRefPose();
 
-	if (CurrentlyUsedMeshComponent.IsValid() && CurrentlyUsedMeshComponent.Get()->SkeletalMesh)
+	if (CurrentlyUsedSourceMeshComponent.IsValid() && CurrentlyUsedSourceMeshComponent.Get()->SkeletalMesh)
 	{
 		const FBoneContainer& RequiredBones = OutPose.GetBoneContainer();
 		for(FCompactPoseBoneIndex PoseBoneIndex : OutPose.ForEachBoneIndex())
@@ -48,18 +48,18 @@ void FAnimNode_CopyPoseFromMesh::Evaluate(FPoseContext& Output)
 			if(Value && *Value!=INDEX_NONE)
 			{
 				const int32 SourceBoneIndex = *Value;
-				const int32 ParentIndex = CurrentlyUsedMeshComponent.Get()->SkeletalMesh->RefSkeleton.GetParentIndex(SourceBoneIndex);
+				const int32 ParentIndex = CurrentlyUsedSourceMeshComponent.Get()->SkeletalMesh->RefSkeleton.GetParentIndex(SourceBoneIndex);
 				const FCompactPoseBoneIndex MyParentIndex = RequiredBones.GetParentBoneIndex(PoseBoneIndex);
 				// only apply if I also have parent, otherwise, it should apply the space bases
 				if (ParentIndex != INDEX_NONE && MyParentIndex != INDEX_NONE)
 				{
-					const FTransform& ParentTransform = CurrentlyUsedMeshComponent.Get()->GetSpaceBases()[ParentIndex];
-					const FTransform& ChildTransform = CurrentlyUsedMeshComponent.Get()->GetSpaceBases()[SourceBoneIndex]; 
+					const FTransform& ParentTransform = CurrentlyUsedSourceMeshComponent.Get()->GetSpaceBases()[ParentIndex];
+					const FTransform& ChildTransform = CurrentlyUsedSourceMeshComponent.Get()->GetSpaceBases()[SourceBoneIndex]; 
 					OutPose[PoseBoneIndex] = ChildTransform.GetRelativeTransform(ParentTransform);
 				}
 				else
 				{
-					OutPose[PoseBoneIndex] = CurrentlyUsedMeshComponent.Get()->GetSpaceBases()[SourceBoneIndex]; 
+					OutPose[PoseBoneIndex] = CurrentlyUsedSourceMeshComponent.Get()->GetSpaceBases()[SourceBoneIndex]; 
 				}
 			}
 		}
@@ -72,7 +72,7 @@ void FAnimNode_CopyPoseFromMesh::GatherDebugData(FNodeDebugData& DebugData)
 
 void FAnimNode_CopyPoseFromMesh::ReinitializeMeshComponent(UAnimInstance* AnimInstance)
 {
-	CurrentlyUsedMeshComponent = SourceMeshComponent;
+	CurrentlyUsedSourceMeshComponent = SourceMeshComponent;
 	BoneMapToSource.Reset();
 	if (SourceMeshComponent && SourceMeshComponent->SkeletalMesh && !SourceMeshComponent->IsPendingKill())
 	{
diff --git a/Engine/Source/Runtime/AnimGraphRuntime/Public/AnimNodes/AnimNode_CopyPoseFromMesh.h b/Engine/Source/Runtime/AnimGraphRuntime/Public/AnimNodes/AnimNode_CopyPoseFromMesh.h
--- a/Engine/Source/Runtime/AnimGraphRuntime/Public/AnimNodes/AnimNode_CopyPoseFromMesh.h
+++ b/Engine/Source/Runtime/AnimGraphRuntime/Public/AnimNodes/AnimNode_CopyPoseFromMesh.h
@@ -28,10 +28,9 @@ struct ANIMGRAPHRUNTIME_API FAnimNode_CopyPoseFromMesh : public FAnimNode_Base
 	// End of FAnimNode_Base interface
 
 private:
-	TWeakObjectPtr<USkeletalMeshComponent> CurrentlyUsedMeshComponent;
+	TWeakObjectPtr<USkeletalMeshComponent> CurrentlyUsedSourceMeshComponent;
 	// cache of target space bases to source space bases
 	TMap<int32, int32> BoneMapToSource;
-	int32 CachedTargetLODIndex;
 
 	// reinitialize mesh component 
 	void ReinitializeMeshComponent(UAnimInstance* AnimInstance);
