#ANIM: Support StartPosition for BlendSpace player

Merging from Orion to Main

diff --git a/Engine/Source/Runtime/AnimGraphRuntime/Private/AnimNodes/AnimNode_BlendSpacePlayer.cpp b/Engine/Source/Runtime/AnimGraphRuntime/Private/AnimNodes/AnimNode_BlendSpacePlayer.cpp
--- a/Engine/Source/Runtime/AnimGraphRuntime/Private/AnimNodes/AnimNode_BlendSpacePlayer.cpp
+++ b/Engine/Source/Runtime/AnimGraphRuntime/Private/AnimNodes/AnimNode_BlendSpacePlayer.cpp
@@ -13,6 +13,7 @@ FAnimNode_BlendSpacePlayer::FAnimNode_BlendSpacePlayer()
 	, Z(0.0f)
 	, PlayRate(1.0f)
 	, bLoop(true)
+	, StartPosition(0.f)
 	, BlendSpace(NULL)
 	, GroupIndex(INDEX_NONE)
 {
@@ -23,11 +24,8 @@ void FAnimNode_BlendSpacePlayer::Initialize(const FAnimationInitializeContext& C
 	BlendSampleDataCache.Empty();
 	
 	EvaluateGraphExposedInputs.Execute(Context);
-	if(PlayRate >= 0.0f)
-	{
-		InternalTimeAccumulator = 0.0f;
-	}
-	else
+	InternalTimeAccumulator = FMath::Clamp(StartPosition, 0.f, 1.0f);
+	if(StartPosition == 0.f && PlayRate < 0.0f)
 	{
 		// Blend spaces run between 0 and 1
 		InternalTimeAccumulator = 1.0f;
diff --git a/Engine/Source/Runtime/AnimGraphRuntime/Public/AnimNodes/AnimNode_BlendSpacePlayer.h b/Engine/Source/Runtime/AnimGraphRuntime/Public/AnimNodes/AnimNode_BlendSpacePlayer.h
--- a/Engine/Source/Runtime/AnimGraphRuntime/Public/AnimNodes/AnimNode_BlendSpacePlayer.h
+++ b/Engine/Source/Runtime/AnimGraphRuntime/Public/AnimNodes/AnimNode_BlendSpacePlayer.h
@@ -33,6 +33,11 @@ public:
 	UPROPERTY(EditAnywhere, BlueprintReadWrite, Category = Settings, meta = (DefaultValue = "true", PinHiddenByDefault))
 	mutable bool bLoop;
 
+	// The start up position in [0, 1], it only applies when reinitialized
+	// if you loop, it will still start from 0.f after finishing the round
+	UPROPERTY(EditAnywhere, BlueprintReadWrite, Category=Settings, meta = (DefaultValue = "0.f", PinHiddenByDefault))
+	mutable float StartPosition;
+
 	// The blendspace asset to play
 	UPROPERTY(EditAnywhere, BlueprintReadWrite, Category = Settings, meta = (PinHiddenByDefault))
 	UBlendSpaceBase* BlendSpace;
