#ANIM: Trail joint - converted TrailRelaxation to curve to support multiple joint chain

Merging

//depot/UE4-Orion/Engine/Source/...

to //depot/UE4/Engine/Source/...

diff --git a/Engine/Source/Editor/AnimGraph/Classes/AnimGraphNode_Trail.h b/Engine/Source/Editor/AnimGraph/Classes/AnimGraphNode_Trail.h
--- a/Engine/Source/Editor/AnimGraph/Classes/AnimGraphNode_Trail.h
+++ b/Engine/Source/Editor/AnimGraph/Classes/AnimGraphNode_Trail.h
@@ -15,6 +15,10 @@ class UAnimGraphNode_Trail : public UAnimGraphNode_SkeletalControlBase
 		FAnimNode_Trail Node;
 
 public:
+	// UObject interface
+	virtual void PostLoad() override;
+	// End of UObject interface
+
 	// UEdGraphNode interface
 	virtual FText GetNodeTitle(ENodeTitleType::Type TitleType) const override;
 //	virtual FString GetNodeNativeTitle(ENodeTitleType::Type TitleType) const override;
diff --git a/Engine/Source/Editor/AnimGraph/Private/AnimGraphNode_Trail.cpp b/Engine/Source/Editor/AnimGraph/Private/AnimGraphNode_Trail.cpp
--- a/Engine/Source/Editor/AnimGraph/Private/AnimGraphNode_Trail.cpp
+++ b/Engine/Source/Editor/AnimGraph/Private/AnimGraphNode_Trail.cpp
@@ -49,4 +49,9 @@ FText UAnimGraphNode_Trail::GetNodeTitle(ENodeTitleType::Type TitleType) const
 	return CachedNodeTitles[TitleType];
 }
 
+void UAnimGraphNode_Trail::PostLoad()
+{
+	Super::PostLoad();
+	Node.PostLoad();
+}
 #undef LOCTEXT_NAMESPACE
\ No newline at end of file
diff --git a/Engine/Source/Editor/DetailCustomizations/Private/AnimTrailNodeDetails.cpp b/Engine/Source/Editor/DetailCustomizations/Private/AnimTrailNodeDetails.cpp
new file mode 100644
--- /dev/null
+++ b/Engine/Source/Editor/DetailCustomizations/Private/AnimTrailNodeDetails.cpp
@@ -0,0 +1,107 @@
+// Copyright 1998-2015 Epic Games, Inc. All Rights Reserved.
+
+#include "DetailCustomizationsPrivatePCH.h"
+
+#include "AnimTrailNodeDetails.h"
+#include "AnimGraphNode_Trail.h"
+
+#define LOCTEXT_NAMESPACE "FAnimTrailNodeDetails"
+
+/////////////////////////////////////////////////////////////////////////
+
+TSharedRef<IDetailCustomization> FAnimTrailNodeDetails::MakeInstance()
+{
+	return MakeShareable( new FAnimTrailNodeDetails );
+}
+
+void FAnimTrailNodeDetails::CustomizeDetails( IDetailLayoutBuilder& DetailBuilder )
+{	
+	TArray<TWeakObjectPtr<UObject> > SelectedObjects;	//the objects we're showing details for
+	DetailBuilder.GetObjectsBeingCustomized(SelectedObjects);
+
+	//we only do fancy customization if we have one vehicle component selected
+	if(SelectedObjects.Num() != 1)
+	{
+		return;
+	}
+	else if(UAnimGraphNode_Trail* InGraphNode = Cast<UAnimGraphNode_Trail>(SelectedObjects[0].Get()))
+	{
+		TrailRelaxCurveEditor = FTrailRelaxCurveEditor(InGraphNode);
+	}
+	else
+	{
+		return;
+	}
+
+	//Trail Relax curve
+	IDetailCategoryBuilder& TrailCategory = DetailBuilder.EditCategory("Trail");
+	TSharedRef<IPropertyHandle> TrailRelaxCurveHandle = DetailBuilder.GetProperty("Node.TrailRelaxationCurve");
+	TSharedPtr<class SCurveEditor> TrailRelaxCurveWidget;
+
+	TrailCategory.AddProperty(TrailRelaxCurveHandle).CustomWidget()
+	.NameContent()
+		[
+			TrailRelaxCurveHandle->CreatePropertyNameWidget()
+		]
+	.ValueContent()
+		.MinDesiredWidth(125.f * 3.f)
+		[
+			SAssignNew(TrailRelaxCurveWidget, SCurveEditor)
+			.ViewMinInput(0.f)
+			.ViewMaxInput(70000.f)
+			.ViewMinOutput(0.f)
+			.ViewMaxOutput(1.f)
+			.TimelineLength(7000.f)
+			.HideUI(false)
+			.DesiredSize(FVector2D(512, 128))
+			.ZoomToFitVertical(true)
+			.ZoomToFitHorizontal(true)
+		];
+
+	TrailRelaxCurveWidget->SetCurveOwner(&TrailRelaxCurveEditor);
+}
+
+
+TArray<FRichCurveEditInfoConst> FAnimTrailNodeDetails::FTrailRelaxCurveEditor::GetCurves() const
+{
+	TArray<FRichCurveEditInfoConst> Curves;
+	Curves.Add(&GraphNodeOwner->Node.TrailRelaxationCurve.EditorCurveData);
+
+	return Curves;
+}
+
+TArray<FRichCurveEditInfo> FAnimTrailNodeDetails::FTrailRelaxCurveEditor::GetCurves()
+{
+	TArray<FRichCurveEditInfo> Curves;
+	Curves.Add(&GraphNodeOwner->Node.TrailRelaxationCurve.EditorCurveData);
+
+	return Curves;
+}
+
+void FAnimTrailNodeDetails::FTrailRelaxCurveEditor::ModifyOwner()
+{
+	if(GraphNodeOwner)
+	{
+		GraphNodeOwner->Modify();
+	}
+}
+
+void FAnimTrailNodeDetails::FTrailRelaxCurveEditor::MakeTransactional()
+{
+	if(GraphNodeOwner)
+	{
+		GraphNodeOwner->SetFlags(GraphNodeOwner->GetFlags() | RF_Transactional);
+	}
+}
+
+bool FAnimTrailNodeDetails::FTrailRelaxCurveEditor::IsValidCurve(FRichCurveEditInfo CurveInfo)
+{
+	return CurveInfo.CurveToEdit == &GraphNodeOwner->Node.TrailRelaxationCurve.EditorCurveData;
+}
+
+FAnimTrailNodeDetails::FTrailRelaxCurveEditor::FTrailRelaxCurveEditor(UAnimGraphNode_Trail * InGraphNode)
+{
+	GraphNodeOwner = InGraphNode;
+}
+
+#undef LOCTEXT_NAMESPACE
diff --git a/Engine/Source/Editor/DetailCustomizations/Private/AnimTrailNodeDetails.h b/Engine/Source/Editor/DetailCustomizations/Private/AnimTrailNodeDetails.h
new file mode 100644
--- /dev/null
+++ b/Engine/Source/Editor/DetailCustomizations/Private/AnimTrailNodeDetails.h
@@ -0,0 +1,33 @@
+// Copyright 1998-2015 Epic Games, Inc. All Rights Reserved.
+
+#pragma once
+
+class UAnimGraphNode_Trail;
+
+class FAnimTrailNodeDetails : public IDetailCustomization
+{
+public:
+	/** Makes a new instance of this detail layout class for a specific detail view requesting it */
+	static TSharedRef<IDetailCustomization> MakeInstance();
+
+	/** IDetailCustomization interface */
+	virtual void CustomizeDetails( IDetailLayoutBuilder& DetailBuilder ) override;
+
+	struct FTrailRelaxCurveEditor: public FCurveOwnerInterface
+	{
+		FTrailRelaxCurveEditor(UAnimGraphNode_Trail* InGraphNode = NULL);
+		/** FCurveOwnerInterface interface */
+		virtual TArray<FRichCurveEditInfoConst> GetCurves() const override;
+		virtual TArray<FRichCurveEditInfo> GetCurves() override;
+		virtual void ModifyOwner() override;
+		virtual void MakeTransactional() override;
+		virtual void OnCurveChanged(const TArray<FRichCurveEditInfo>& ChangedCurveEditInfos) override { }
+		virtual bool IsValidCurve(FRichCurveEditInfo CurveInfo) override;
+
+	private:
+		UAnimGraphNode_Trail* GraphNodeOwner;
+
+	} TrailRelaxCurveEditor;
+
+};
+
diff --git a/Engine/Source/Editor/DetailCustomizations/Private/DetailCustomizations.cpp b/Engine/Source/Editor/DetailCustomizations/Private/DetailCustomizations.cpp
--- a/Engine/Source/Editor/DetailCustomizations/Private/DetailCustomizations.cpp
+++ b/Engine/Source/Editor/DetailCustomizations/Private/DetailCustomizations.cpp
@@ -104,6 +104,7 @@
 #include "ConfigEditorPropertyDetails.h"
 #include "AssetImportDataCustomization.h"
 #include "TextCustomization.h"
+#include "AnimTrailNodeDetails.h"
 
 IMPLEMENT_MODULE( FDetailCustomizationsModule, DetailCustomizations );
 
@@ -259,6 +260,7 @@ void FDetailCustomizationsModule::RegisterObjectCustomizations()
 	RegisterCustomClassLayout("AnimStateNode", FOnGetDetailCustomizationInstance::CreateStatic(&FAnimStateNodeDetails::MakeInstance));
 	RegisterCustomClassLayout("AnimStateTransitionNode", FOnGetDetailCustomizationInstance::CreateStatic(&FAnimTransitionNodeDetails::MakeInstance));
 	RegisterCustomClassLayout("AnimGraphNode_TransitionPoseEvaluator", FOnGetDetailCustomizationInstance::CreateStatic(&FTransitionPoseEvaluatorNodeDetails::MakeInstance));
+	RegisterCustomClassLayout("AnimGraphNode_Trail", FOnGetDetailCustomizationInstance::CreateStatic(&FAnimTrailNodeDetails::MakeInstance));
 
 	RegisterCustomClassLayout("SoundWave", FOnGetDetailCustomizationInstance::CreateStatic(&FSoundWaveDetails::MakeInstance));
 	RegisterCustomClassLayout("DialogueWave", FOnGetDetailCustomizationInstance::CreateStatic(&FDialogueWaveDetails::MakeInstance));
diff --git a/Engine/Source/Runtime/AnimGraphRuntime/Private/BoneControllers/AnimNode_Trail.cpp b/Engine/Source/Runtime/AnimGraphRuntime/Private/BoneControllers/AnimNode_Trail.cpp
--- a/Engine/Source/Runtime/AnimGraphRuntime/Private/BoneControllers/AnimNode_Trail.cpp
+++ b/Engine/Source/Runtime/AnimGraphRuntime/Private/BoneControllers/AnimNode_Trail.cpp
@@ -11,12 +11,14 @@ FAnimNode_Trail::FAnimNode_Trail()
 	, ChainBoneAxis(EAxis::X)
 	, bInvertChainBoneAxis(false)
 	, bLimitStretch(false)
-	, TrailRelaxation(10)
+	, TrailRelaxation_DEPRECATED(10.f)
 	, StretchLimit(0)
 	, FakeVelocity(FVector::ZeroVector)
 	, bActorSpaceFakeVel(false)
 	, bHadValidStrength(false)
 {
+	FRichCurve* TrailRelaxRichCurve = TrailRelaxationCurve.GetRichCurve();
+	TrailRelaxRichCurve->AddKey(0.f, 10.f);
 }
 
 void FAnimNode_Trail::Update(const FAnimationUpdateContext& Context)
@@ -42,20 +44,19 @@ void FAnimNode_Trail::EvaluateBoneTransforms(USkeletalMeshComponent* SkelComp, F
 {
 	check(OutBoneTransforms.Num() == 0);
 
-	if( ChainLength < 2 )
+	if( ChainBoneIndices.Num() <= 0 )
 	{
 		return;
 	}
 
+	checkSlow (ChainBoneIndices.Num() == ChainLength);
+	checkSlow (PerJointTrailData.Num() == ChainLength);
+
 	// The incoming BoneIndex is the 'end' of the spline chain. We need to find the 'start' by walking SplineLength bones up hierarchy.
 	// Fail if we walk past the root bone.
-
 	const FBoneContainer& BoneContainer = MeshBases.GetPose().GetBoneContainer();
 	FCompactPoseBoneIndex WalkBoneIndex = TrailBone.GetCompactPoseIndex(BoneContainer);
 
-	TArray<FCompactPoseBoneIndex> ChainBoneIndices;
-	ChainBoneIndices.AddZeroed(ChainLength);
-
 	ChainBoneIndices[ChainLength - 1] = WalkBoneIndex;
 
 	for (int32 i = 1; i < ChainLength; i++)
@@ -80,11 +81,8 @@ void FAnimNode_Trail::EvaluateBoneTransforms(USkeletalMeshComponent* SkelComp, F
 	// If we have >0 this frame, but didn't last time, record positions of all the bones.
 	// Also do this if number has changed or array is zero.
 	bool bHasValidStrength = (Alpha > 0.f);
-	if(TrailBoneLocations.Num() != ChainLength || (bHasValidStrength && !bHadValidStrength))
+	if(bHasValidStrength && !bHadValidStrength)
 	{
-		TrailBoneLocations.Empty();
-		TrailBoneLocations.AddZeroed(ChainLength);
-
 		for(int32 i=0; i<ChainBoneIndices.Num(); i++)
 		{
 			FCompactPoseBoneIndex ChildIndex = ChainBoneIndices[i];
@@ -124,6 +122,7 @@ void FAnimNode_Trail::EvaluateBoneTransforms(USkeletalMeshComponent* SkelComp, F
 	TrailBoneLocations[0] = ChainTransform.GetTranslation();
 
 	// Starting one below head of chain, move bones.
+	// this Parent/Child relationship is backward. From start joint (from bottom) to end joint(higher parent )
 	for(int32 i=1; i<ChainBoneIndices.Num(); i++)
 	{
 		// Parent bone position in component space.
@@ -146,7 +145,7 @@ void FAnimNode_Trail::EvaluateBoneTransforms(USkeletalMeshComponent* SkelComp, F
 		FVector Error = ChildTarget - ChildPos;
 
 		// Calculate how much to push the child towards its target
-		float Correction = FMath::Clamp<float>(ThisTimstep * TrailRelaxation, 0.f, 1.f);
+		float Correction = FMath::Clamp<float>(ThisTimstep * PerJointTrailData[i].TrailRelaxationSpeedPerSecond, 0.f, 1.f);
 
 		// Scale correction vector and apply to get new world-space child position.
 		TrailBoneLocations[i] = ChildPos + (Error * Correction);
@@ -232,3 +231,39 @@ FVector FAnimNode_Trail::GetAlignVector(EAxis::Type AxisOption, bool bInvert)
 
 	return AxisDir;
 }
+
+void FAnimNode_Trail::PostLoad()
+{
+	if (TrailRelaxation_DEPRECATED != 10.f)
+	{
+		FRichCurve* TrailRelaxRichCurve = TrailRelaxationCurve.GetRichCurve();
+		TrailRelaxRichCurve->Reset();
+		TrailRelaxRichCurve->AddKey(0.f, TrailRelaxation_DEPRECATED);
+		// since we don't know if it's same as default or not, we have to keep default
+		// if default, the default constructor will take care of it. If not, we'll reset
+		TrailRelaxation_DEPRECATED = 10.f;
+	}
+}
+void FAnimNode_Trail::Initialize(const FAnimationInitializeContext& Context)
+{
+	FAnimNode_SkeletalControlBase::Initialize(Context);
+
+	// allocated all memory here in initialize
+	PerJointTrailData.Reset();
+	ChainBoneIndices.Reset();
+	TrailBoneLocations.Reset();
+	if(ChainLength > 2)
+	{
+		PerJointTrailData.AddZeroed(ChainLength);
+		ChainBoneIndices.AddZeroed(ChainLength);
+		TrailBoneLocations.AddZeroed(ChainLength);
+
+		float Interval = (ChainLength > 1)? (1.f/(ChainLength-1)) : 0.f;
+		const FRichCurve* TrailRelaxRichCurve = TrailRelaxationCurve.GetRichCurveConst();
+		ensure(TrailRelaxRichCurve);
+		for(int32 Idx=0; Idx<ChainLength; ++Idx)
+		{
+			PerJointTrailData[Idx].TrailRelaxationSpeedPerSecond = TrailRelaxRichCurve->Eval(Interval * Idx);
+		}
+	}
+}
\ No newline at end of file
diff --git a/Engine/Source/Runtime/AnimGraphRuntime/Public/BoneControllers/AnimNode_Trail.h b/Engine/Source/Runtime/AnimGraphRuntime/Public/BoneControllers/AnimNode_Trail.h
--- a/Engine/Source/Runtime/AnimGraphRuntime/Public/BoneControllers/AnimNode_Trail.h
+++ b/Engine/Source/Runtime/AnimGraphRuntime/Public/BoneControllers/AnimNode_Trail.h
@@ -5,6 +5,14 @@
 #include "AnimNode_SkeletalControlBase.h"
 #include "AnimNode_Trail.generated.h"
 
+// in the future, we might use this for stretch set up as well
+// for now this is unserializable, and transient only
+struct FPerJointTrailSetup
+{
+	/** How quickly we 'relax' the bones to their animated positions. */
+	float	TrailRelaxationSpeedPerSecond;
+};
+
 /**
  * Trail Controller
  */
@@ -34,9 +42,13 @@ struct ANIMGRAPHRUNTIME_API FAnimNode_Trail : public FAnimNode_SkeletalControlBa
 	UPROPERTY(EditAnywhere, BlueprintReadWrite, Category=Trail)
 	bool	bLimitStretch;
 
-	/** How quickly we 'relax' the bones to their animated positions. */
+	/** How quickly we 'relax' the bones to their animated positions. Deprecated. Replaced to TrailRelaxationCurve */
+	UPROPERTY()
+	float	TrailRelaxation_DEPRECATED;
+
+	/** How quickly we 'relax' the bones to their animated positions. Time 0 will map to start joint and time 1 will map to the last chain bone. */
 	UPROPERTY(EditAnywhere, BlueprintReadWrite, Category=Trail)
-	float	TrailRelaxation;
+	FRuntimeFloatCurve TrailRelaxationCurve;
 
 	/** If bLimitStretch is true, this indicates how long a bone can stretch beyond its length in the ref-pose. */
 	UPROPERTY(EditAnywhere, BlueprintReadWrite, Category=Trail)
@@ -62,10 +74,13 @@ struct ANIMGRAPHRUNTIME_API FAnimNode_Trail : public FAnimNode_SkeletalControlBa
 	/** LocalToWorld used last frame, used for building transform between frames. */
 	FMatrix		OldLocalToWorld;
 
+	/** Per Joint Trail Set up*/
+	TArray<FPerJointTrailSetup> PerJointTrailData;
 
 	FAnimNode_Trail();
 
 	// FAnimNode_Base interface
+	virtual void Initialize(const FAnimationInitializeContext& Context) override;
 	virtual void Update(const FAnimationUpdateContext& Context) override;
 	virtual void GatherDebugData(FNodeDebugData& DebugData) override;
 	// End of FAnimNode_Base interface
@@ -75,10 +90,14 @@ struct ANIMGRAPHRUNTIME_API FAnimNode_Trail : public FAnimNode_SkeletalControlBa
 	virtual bool IsValidToEvaluate(const USkeleton* Skeleton, const FBoneContainer& RequiredBones) override;
 	// End of FAnimNode_SkeletalControlBase interface
 
+	void PostLoad();
+
 private:
 	// FAnimNode_SkeletalControlBase interface
 	virtual void InitializeBoneReferences(const FBoneContainer& RequiredBones) override;
 	// End of FAnimNode_SkeletalControlBase interface
 
 	FVector GetAlignVector(EAxis::Type AxisOption, bool bInvert);
+
+	TArray<FCompactPoseBoneIndex> ChainBoneIndices;
 };
\ No newline at end of file
