Merging

//depot/UE4-Orion/Engine/Source/Runtime/AnimGraphRuntime/Private/BoneControllers/AnimNode_Trail.cpp

to //depot/UE4/Engine/Source/Runtime/AnimGraphRuntime/Private/BoneControllers/AnimNode_Trail.cpp

diff --git a/Engine/Source/Runtime/AnimGraphRuntime/Private/BoneControllers/AnimNode_Trail.cpp b/Engine/Source/Runtime/AnimGraphRuntime/Private/BoneControllers/AnimNode_Trail.cpp
--- a/Engine/Source/Runtime/AnimGraphRuntime/Private/BoneControllers/AnimNode_Trail.cpp
+++ b/Engine/Source/Runtime/AnimGraphRuntime/Private/BoneControllers/AnimNode_Trail.cpp
@@ -19,6 +19,7 @@ FAnimNode_Trail::FAnimNode_Trail()
 {
 	FRichCurve* TrailRelaxRichCurve = TrailRelaxationCurve.GetRichCurve();
 	TrailRelaxRichCurve->AddKey(0.f, 10.f);
+	TrailRelaxRichCurve->AddKey(1.f, 10.f);
 }
 
 void FAnimNode_Trail::Update(const FAnimationUpdateContext& Context)
@@ -239,6 +240,7 @@ void FAnimNode_Trail::PostLoad()
 		FRichCurve* TrailRelaxRichCurve = TrailRelaxationCurve.GetRichCurve();
 		TrailRelaxRichCurve->Reset();
 		TrailRelaxRichCurve->AddKey(0.f, TrailRelaxation_DEPRECATED);
+		TrailRelaxRichCurve->AddKey(1.f, TrailRelaxation_DEPRECATED);
 		// since we don't know if it's same as default or not, we have to keep default
 		// if default, the default constructor will take care of it. If not, we'll reset
 		TrailRelaxation_DEPRECATED = 10.f;
