#ue4 - Fix for crash running QAGame PIE 100% of the time in anim node code in debug.

diff --git a/Engine/Source/Runtime/AnimGraphRuntime/Private/AnimNodes/AnimNode_BlendSpacePlayer.cpp b/Engine/Source/Runtime/AnimGraphRuntime/Private/AnimNodes/AnimNode_BlendSpacePlayer.cpp
--- a/Engine/Source/Runtime/AnimGraphRuntime/Private/AnimNodes/AnimNode_BlendSpacePlayer.cpp
+++ b/Engine/Source/Runtime/AnimGraphRuntime/Private/AnimNodes/AnimNode_BlendSpacePlayer.cpp
@@ -21,6 +21,8 @@ FAnimNode_BlendSpacePlayer::FAnimNode_BlendSpacePlayer()
 
 void FAnimNode_BlendSpacePlayer::Initialize(const FAnimationInitializeContext& Context)
 {
+	FAnimNode_AssetPlayerBase::Initialize(Context);
+
 	BlendSampleDataCache.Empty();
 	
 	EvaluateGraphExposedInputs.Execute(Context);
