#ANIM : Fix for slot node weight inconsistency issue

Merging

//depot/UE4-Orion/Engine/Source/Runtime/Engine/Private/Animation/AnimNode_Slot.cpp

to //depot/UE4/Engine/Source/Runtime/AnimGraphRuntime/Private/AnimNodes/AnimNode_Slot.cpp

diff --git a/Engine/Source/Runtime/AnimGraphRuntime/Private/AnimNodes/AnimNode_Slot.cpp b/Engine/Source/Runtime/AnimGraphRuntime/Private/AnimNodes/AnimNode_Slot.cpp
--- a/Engine/Source/Runtime/AnimGraphRuntime/Private/AnimNodes/AnimNode_Slot.cpp
+++ b/Engine/Source/Runtime/AnimGraphRuntime/Private/AnimNodes/AnimNode_Slot.cpp
@@ -55,7 +55,8 @@ void FAnimNode_Slot::Evaluate(FPoseContext & Output)
 	else
 	{
 		FPoseContext SourceContext(Output);
-		if (SourceWeight > ZERO_ANIMWEIGHT_THRESH)
+// @TODO: FIXME temporary fix to fix crash 
+//		if (SourceWeight > ZERO_ANIMWEIGHT_THRESH)
 		{
 			Source.Evaluate(SourceContext);
 		}
