Fix compile error.
Merging

//depot/UE4-Orion/Engine/Source/Runtime/Engine/Private/Animation/AnimNode_BlendListBase.cpp

to //depot/UE4/Engine/Source/Runtime/AnimGraphRuntime/Private/AnimNodes/AnimNode_BlendListBase.cpp

diff --git a/Engine/Source/Runtime/AnimGraphRuntime/Private/AnimNodes/AnimNode_BlendListBase.cpp b/Engine/Source/Runtime/AnimGraphRuntime/Private/AnimNodes/AnimNode_BlendListBase.cpp
--- a/Engine/Source/Runtime/AnimGraphRuntime/Private/AnimNodes/AnimNode_BlendListBase.cpp
+++ b/Engine/Source/Runtime/AnimGraphRuntime/Private/AnimNodes/AnimNode_BlendListBase.cpp
@@ -40,7 +40,7 @@ void FAnimNode_BlendListBase::Initialize(const FAnimationInitializeContext& Cont
 		FAlphaBlend& Blend = Blends[i];
 
 		Blend.SetBlendTime(0.0f);
-		Blend.SetBlendOption(BlendType);
+		Blend.BlendOption = BlendType;
 		Blend.CustomCurve = CustomBlendCurve;
 	}
 	Blends[0].SetAlpha(1.0f);
