#ANIM: trail

- Cleaned up code, and added base joint
 - base joint doesn't work but not sure exactly why yet. Need to investigate, but at least this cleans up some stuff from trail.

Merging

//depot/UE4-Orion/Engine/Source/Runtime/...

to //depot/UE4-Orion/Engine/Source/Runtime/...

diff --git a/Engine/Source/Runtime/AnimGraphRuntime/Private/BoneControllers/AnimNode_Trail.cpp b/Engine/Source/Runtime/AnimGraphRuntime/Private/BoneControllers/AnimNode_Trail.cpp
--- a/Engine/Source/Runtime/AnimGraphRuntime/Private/BoneControllers/AnimNode_Trail.cpp
+++ b/Engine/Source/Runtime/AnimGraphRuntime/Private/BoneControllers/AnimNode_Trail.cpp
@@ -10,8 +10,8 @@ FAnimNode_Trail::FAnimNode_Trail()
 	: ChainLength(2)
 	, ChainBoneAxis(EAxis::X)
 	, bInvertChainBoneAxis(false)
-	, bLimitStretch(false)
 	, TrailRelaxation_DEPRECATED(10.f)
+	, bLimitStretch(false)
 	, StretchLimit(0)
 	, FakeVelocity(FVector::ZeroVector)
 	, bActorSpaceFakeVel(false)
@@ -52,50 +52,53 @@ void FAnimNode_Trail::EvaluateBoneTransforms(USkeletalMeshComponent* SkelComp, F
 
 	checkSlow (ChainBoneIndices.Num() == ChainLength);
 	checkSlow (PerJointTrailData.Num() == ChainLength);
-
 	// The incoming BoneIndex is the 'end' of the spline chain. We need to find the 'start' by walking SplineLength bones up hierarchy.
 	// Fail if we walk past the root bone.
 	const FBoneContainer& BoneContainer = MeshBases.GetPose().GetBoneContainer();
-	FCompactPoseBoneIndex WalkBoneIndex = TrailBone.GetCompactPoseIndex(BoneContainer);
-
-	ChainBoneIndices[ChainLength - 1] = WalkBoneIndex;
-
-	for (int32 i = 1; i < ChainLength; i++)
+	const FTransform ComponentTransform = SkelComp->GetComponentToWorld();
+	FTransform BaseTransform;
+ 	if (BaseJoint.IsValid(BoneContainer)) 
+ 	{
+		FCompactPoseBoneIndex BasePoseIndex = BoneContainer.MakeCompactPoseIndex(FMeshPoseBoneIndex(BaseJoint.BoneIndex));
+		FTransform BaseBoneTransform = MeshBases.GetComponentSpaceTransform(BasePoseIndex);
+ 		BaseTransform = BaseBoneTransform * ComponentTransform;
+ 	}
+	else
 	{
-		// returns to avoid a crash
-		// @TODO : shows an error message why failed
-		if (WalkBoneIndex == 0)
-		{
-			return;
-		}
-
-		// Get parent bone.
-		WalkBoneIndex = BoneContainer.GetParentBoneIndex(WalkBoneIndex);
-
-		//Insert indices at the start of array, so that parents are before children in the array.
-		int32 TransformIndex = ChainLength - (i + 1);
-		ChainBoneIndices[TransformIndex] = WalkBoneIndex;
+		BaseTransform = ComponentTransform;
 	}
 
 	OutBoneTransforms.AddZeroed(ChainLength);
 
+	// this should be checked outside
+	checkSlow (TrailBone.IsValid(BoneContainer));
+
 	// If we have >0 this frame, but didn't last time, record positions of all the bones.
 	// Also do this if number has changed or array is zero.
+	//@todo I don't think this will work anymore. if Alpha is too small, it won't call evaluate anyway
+	// so this has to change. AFAICT, this will get called only FIRST TIME
 	bool bHasValidStrength = (Alpha > 0.f);
 	if(bHasValidStrength && !bHadValidStrength)
 	{
 		for(int32 i=0; i<ChainBoneIndices.Num(); i++)
 		{
-			FCompactPoseBoneIndex ChildIndex = ChainBoneIndices[i];
-			const FTransform& ChainTransform = MeshBases.GetComponentSpaceTransform(ChildIndex);
-			TrailBoneLocations[i] = ChainTransform.GetTranslation();
+			if (BoneContainer.Contains(ChainBoneIndices[i]))
+			{
+				FCompactPoseBoneIndex ChildIndex = BoneContainer.MakeCompactPoseIndex(FMeshPoseBoneIndex(ChainBoneIndices[i]));
+				const FTransform& ChainTransform = MeshBases.GetComponentSpaceTransform(ChildIndex);
+				TrailBoneLocations[i] = ChainTransform.GetTranslation();
+			}
+			else
+			{
+				TrailBoneLocations[i] = FVector::ZeroVector;
+			}
 		}
-		OldLocalToWorld = SkelComp->GetTransformMatrix();
+		OldBaseTransform = BaseTransform;
 	}
 	bHadValidStrength = bHasValidStrength;
 
 	// transform between last frame and now.
-	FMatrix OldToNewTM = OldLocalToWorld * SkelComp->GetTransformMatrix().InverseFast();
+	FTransform OldToNewTM = OldBaseTransform.GetRelativeTransform(BaseTransform);
 
 	// Add fake velocity if present to all but root bone
 	if(!FakeVelocity.IsZero())
@@ -108,7 +111,7 @@ void FAnimNode_Trail::EvaluateBoneTransforms(USkeletalMeshComponent* SkelComp, F
 			FakeMovement = BoneToWorld.TransformVector(FakeMovement);
 		}
 
-		FakeMovement = SkelComp->GetTransformMatrix().InverseTransformVector(FakeMovement);
+		FakeMovement = BaseTransform.InverseTransformVector(FakeMovement);
 		// Then add to each bone
 		for(int32 i=1; i<TrailBoneLocations.Num(); i++)
 		{
@@ -117,7 +120,7 @@ void FAnimNode_Trail::EvaluateBoneTransforms(USkeletalMeshComponent* SkelComp, F
 	}
 
 	// Root bone of trail is not modified.
-	FCompactPoseBoneIndex RootIndex = ChainBoneIndices[0]; 
+	FCompactPoseBoneIndex RootIndex = BoneContainer.MakeCompactPoseIndex(FMeshPoseBoneIndex(ChainBoneIndices[0])); 
 	const FTransform& ChainTransform = MeshBases.GetComponentSpaceTransform(RootIndex);
 	OutBoneTransforms[0] = FBoneTransform(RootIndex, ChainTransform);
 	TrailBoneLocations[0] = ChainTransform.GetTranslation();
@@ -127,12 +130,12 @@ void FAnimNode_Trail::EvaluateBoneTransforms(USkeletalMeshComponent* SkelComp, F
 	for(int32 i=1; i<ChainBoneIndices.Num(); i++)
 	{
 		// Parent bone position in component space.
-		FCompactPoseBoneIndex ParentIndex = ChainBoneIndices[i - 1];
+		FCompactPoseBoneIndex ParentIndex = BoneContainer.MakeCompactPoseIndex(FMeshPoseBoneIndex(ChainBoneIndices[i - 1]));
 		FVector ParentPos = TrailBoneLocations[i-1];
 		FVector ParentAnimPos = MeshBases.GetComponentSpaceTransform(ParentIndex).GetTranslation();
 
 		// Child bone position in component space.
-		FCompactPoseBoneIndex ChildIndex = ChainBoneIndices[i];
+		FCompactPoseBoneIndex ChildIndex = BoneContainer.MakeCompactPoseIndex(FMeshPoseBoneIndex(ChainBoneIndices[i]));
 		FVector ChildPos = OldToNewTM.TransformPosition(TrailBoneLocations[i]); // move from 'last frames component' frame to 'this frames component' frame
 		FVector ChildAnimPos = MeshBases.GetComponentSpaceTransform(ChildIndex).GetTranslation();
 
@@ -184,28 +187,67 @@ void FAnimNode_Trail::EvaluateBoneTransforms(USkeletalMeshComponent* SkelComp, F
 		FTransform DeltaTM( DeltaLookQuat, FVector(0.f) );
 
 		// Apply to the current parent bone transform.
-		FTransform TmpMatrix = FTransform::Identity;
-		TmpMatrix.CopyRotationPart(OutBoneTransforms[i - 1].Transform);
-		TmpMatrix = TmpMatrix * DeltaTM;
-		OutBoneTransforms[i - 1].Transform.CopyRotationPart(TmpMatrix);
+		FTransform TmpTransform = FTransform::Identity;
+		TmpTransform.CopyRotationPart(OutBoneTransforms[i - 1].Transform);
+		TmpTransform = TmpTransform * DeltaTM;
+		OutBoneTransforms[i - 1].Transform.CopyRotationPart(TmpTransform);
 	}
 
 	// For the last bone in the chain, use the rotation from the bone above it.
 	OutBoneTransforms[ChainLength - 1].Transform.CopyRotationPart(OutBoneTransforms[ChainLength - 2].Transform);
 
-	// Update OldLocalToWorld
-	OldLocalToWorld = SkelComp->GetTransformMatrix();
+	// Update OldBaseTransform
+	OldBaseTransform = BaseTransform;
 }
 
 bool FAnimNode_Trail::IsValidToEvaluate(const USkeleton* Skeleton, const FBoneContainer& RequiredBones) 
 {
 	// if bones are valid
-	return (TrailBone.IsValid(RequiredBones));
+	if (TrailBone.IsValid(RequiredBones))
+	{
+		for (auto& ChainIndex : ChainBoneIndices)
+		{
+			if (RequiredBones.Contains(ChainIndex) == false)
+			{
+				return false;
+			}
+		}
+	}
+
+	return true;
 }
 
 void FAnimNode_Trail::InitializeBoneReferences(const FBoneContainer& RequiredBones) 
 {
 	TrailBone.Initialize(RequiredBones);
+
+	// initialize chain bone indices
+	ChainBoneIndices.Reset();
+	if (ChainLength > 2 && TrailBone.IsValid(RequiredBones))
+	{
+		ChainBoneIndices.AddZeroed(ChainLength);
+
+		int32 WalkBoneIndex = TrailBone.BoneIndex;
+		ChainBoneIndices[ChainLength - 1] = WalkBoneIndex;
+
+		for(int32 i = 1; i < ChainLength; i++)
+		{
+			//Insert indices at the start of array, so that parents are before children in the array.
+			int32 TransformIndex = ChainLength - (i + 1);
+
+			// if reached to root or invalid, invalidate the data
+			if(WalkBoneIndex == INDEX_NONE || WalkBoneIndex == 0)
+			{
+				ChainBoneIndices[TransformIndex] = INDEX_NONE;
+			}
+			else
+			{
+				// Get parent bone.
+				WalkBoneIndex = RequiredBones.GetParentBoneIndex(WalkBoneIndex);
+				ChainBoneIndices[TransformIndex] = WalkBoneIndex;
+			}
+		}
+	}
 }
 
 FVector FAnimNode_Trail::GetAlignVector(EAxis::Type AxisOption, bool bInvert)
@@ -250,14 +292,17 @@ void FAnimNode_Trail::Initialize(const FAnimationInitializeContext& Context)
 {
 	FAnimNode_SkeletalControlBase::Initialize(Context);
 
+	const USkeleton* Skeleton = Context.AnimInstance->CurrentSkeleton;
+	check (Skeleton);
+
+	BaseJoint.Initialize(Skeleton);
+
 	// allocated all memory here in initialize
 	PerJointTrailData.Reset();
-	ChainBoneIndices.Reset();
 	TrailBoneLocations.Reset();
 	if(ChainLength > 2)
 	{
 		PerJointTrailData.AddZeroed(ChainLength);
-		ChainBoneIndices.AddZeroed(ChainLength);
 		TrailBoneLocations.AddZeroed(ChainLength);
 
 		float Interval = (ChainLength > 1)? (1.f/(ChainLength-1)) : 0.f;
diff --git a/Engine/Source/Runtime/AnimGraphRuntime/Public/BoneControllers/AnimNode_Trail.h b/Engine/Source/Runtime/AnimGraphRuntime/Public/BoneControllers/AnimNode_Trail.h
--- a/Engine/Source/Runtime/AnimGraphRuntime/Public/BoneControllers/AnimNode_Trail.h
+++ b/Engine/Source/Runtime/AnimGraphRuntime/Public/BoneControllers/AnimNode_Trail.h
@@ -38,10 +38,6 @@ struct ANIMGRAPHRUNTIME_API FAnimNode_Trail : public FAnimNode_SkeletalControlBa
 	UPROPERTY(EditAnywhere, BlueprintReadWrite, Category=Trail)
 	bool	bInvertChainBoneAxis;
 
-	/** Limit the amount that a bone can stretch from its ref-pose length. */
-	UPROPERTY(EditAnywhere, BlueprintReadWrite, Category=Trail)
-	bool	bLimitStretch;
-
 	/** How quickly we 'relax' the bones to their animated positions. Deprecated. Replaced to TrailRelaxationCurve */
 	UPROPERTY()
 	float	TrailRelaxation_DEPRECATED;
@@ -50,18 +46,26 @@ struct ANIMGRAPHRUNTIME_API FAnimNode_Trail : public FAnimNode_SkeletalControlBa
 	UPROPERTY(EditAnywhere, BlueprintReadWrite, Category=Trail, meta=(CustomizeProperty))
 	FRuntimeFloatCurve TrailRelaxationSpeed;
 
+	/** Limit the amount that a bone can stretch from its ref-pose length. */
+	UPROPERTY(EditAnywhere, BlueprintReadWrite, Category=Limit)
+	bool	bLimitStretch;
+
 	/** If bLimitStretch is true, this indicates how long a bone can stretch beyond its length in the ref-pose. */
-	UPROPERTY(EditAnywhere, BlueprintReadWrite, Category=Trail)
+	UPROPERTY(EditAnywhere, BlueprintReadWrite, Category=Limit)
 	float	StretchLimit;
 
 	/** 'Fake' velocity applied to bones. */
-	UPROPERTY(EditAnywhere, BlueprintReadWrite, Category=Trail)
+	UPROPERTY(EditAnywhere, BlueprintReadWrite, Category=Velocity)
 	FVector	FakeVelocity;
 
 	/** Whether 'fake' velocity should be applied in actor or world space. */
-	UPROPERTY(EditAnywhere, BlueprintReadWrite, Category=Trail)
+	UPROPERTY(EditAnywhere, BlueprintReadWrite, Category=Velocity)
 	bool	bActorSpaceFakeVel;
 
+	/** Base Joint to calculate velocity from. If none, it will use Component's World Transform. . */
+	UPROPERTY(EditAnywhere, BlueprintReadWrite, Category=Velocity)
+	FBoneReference BaseJoint;
+
 	/** Internal use - we need the timestep to do the relaxation in CalculateNewBoneTransforms. */
 	float	ThisTimstep;
 
@@ -72,7 +76,7 @@ struct ANIMGRAPHRUNTIME_API FAnimNode_Trail : public FAnimNode_SkeletalControlBa
 	TArray<FVector>	TrailBoneLocations;
 
 	/** LocalToWorld used last frame, used for building transform between frames. */
-	FMatrix		OldLocalToWorld;
+	FTransform		OldBaseTransform;
 
 	/** Per Joint Trail Set up*/
 	TArray<FPerJointTrailSetup> PerJointTrailData;
@@ -99,5 +103,6 @@ private:
 
 	FVector GetAlignVector(EAxis::Type AxisOption, bool bInvert);
 
-	TArray<FCompactPoseBoneIndex> ChainBoneIndices;
+	// skeleton index
+	TArray<FBoneIndexType> ChainBoneIndices;
 };
\ No newline at end of file
diff --git a/Engine/Source/Runtime/Engine/Classes/Components/SkeletalMeshComponent.h b/Engine/Source/Runtime/Engine/Classes/Components/SkeletalMeshComponent.h
--- a/Engine/Source/Runtime/Engine/Classes/Components/SkeletalMeshComponent.h
+++ b/Engine/Source/Runtime/Engine/Classes/Components/SkeletalMeshComponent.h
@@ -1024,13 +1024,6 @@ public:
 
 	void GetWeldedBodies(TArray<FBodyInstance*> & OutWeldedBodies, TArray<FName> & OutChildrenLabels) override;
 
-	/**
-	 * Return Transform Matrix for SkeletalMeshComponent considering root motion setups
-	 * 
-	 * @return Matrix Transform matrix
-	 */
-	FMatrix GetTransformMatrix() const;
-	
 	/** 
 	 * Change whether to force mesh into ref pose (and use cheaper vertex shader) 
 	 *
diff --git a/Engine/Source/Runtime/Engine/Private/Components/SkeletalMeshComponent.cpp b/Engine/Source/Runtime/Engine/Private/Components/SkeletalMeshComponent.cpp
--- a/Engine/Source/Runtime/Engine/Private/Components/SkeletalMeshComponent.cpp
+++ b/Engine/Source/Runtime/Engine/Private/Components/SkeletalMeshComponent.cpp
@@ -1455,35 +1455,6 @@ UAnimInstance* USkeletalMeshComponent::GetAnimInstance() const
 	return AnimScriptInstance;
 }
 
-FMatrix USkeletalMeshComponent::GetTransformMatrix() const
-{
-	FTransform RootTransform = GetBoneTransform(0);
-	FVector Translation;
-	FQuat Rotation;
-	
-	// if in editor, it should always use localToWorld
-	// if root motion is ignored, use root transform 
-	if( GetWorld()->IsGameWorld() || !SkeletalMesh )
-	{
-		// add root translation info
-		Translation = RootTransform.GetLocation();
-	}
-	else
-	{
-		Translation = ComponentToWorld.TransformPosition(SkeletalMesh->RefSkeleton.GetRefBonePose()[0].GetTranslation());
-	}
-
-	// if root rotation is ignored, use root transform rotation
-	Rotation = RootTransform.GetRotation();
-
-	// now I need to get scale
-	// only LocalToWorld will have scale
-	FVector ScaleVector = ComponentToWorld.GetScale3D();
-
-	Rotation.Normalize();
-	return FScaleMatrix(ScaleVector)*FQuatRotationTranslationMatrix(Rotation, Translation);
-}
-
 void USkeletalMeshComponent::NotifySkelControlBeyondLimit( USkelControlLookAt* LookAt ) {}
 
 
