Copying //UE4/Dev-Framework to //UE4/Dev-Main (Source: //UE4/Dev-Framework @ 3021292)

diff --git a/Engine/Source/Editor/Persona/Private/SAnimViewportToolBar.cpp b/Engine/Source/Editor/Persona/Private/SAnimViewportToolBar.cpp
--- a/Engine/Source/Editor/Persona/Private/SAnimViewportToolBar.cpp
+++ b/Engine/Source/Editor/Persona/Private/SAnimViewportToolBar.cpp
@@ -18,6 +18,8 @@
 #include "SNumericEntryBox.h"
 #include "STextComboBox.h"
 
+#include "AssetViewerSettings.h"
+
 
 #define LOCTEXT_NAMESPACE "AnimViewportToolBar"
 
@@ -121,9 +123,13 @@ protected:
 	/** Callback function which determines whether this widget is enabled */
 	bool IsBrightnessSliderEnabled() const
 	{
-		// Only enable the brightness control when we can see the background (ie, no sky or floor)
-		TSharedPtr<SAnimationEditorViewportTabBody> AnimViewportPinnedPtr = AnimViewportPtr.Pin();
-		return !(AnimViewportPinnedPtr->IsShowingSky() || AnimViewportPinnedPtr->IsShowingFloor());
+		const UAssetViewerSettings* Settings = UAssetViewerSettings::Get();
+		const UEditorPerProjectUserSettings* PerProjectUserSettings = GetDefault<UEditorPerProjectUserSettings>();
+		const int32 ProfileIndex = Settings->Profiles.IsValidIndex(PerProjectUserSettings->AssetViewerProfileIndex) ? PerProjectUserSettings->AssetViewerProfileIndex : 0;
+
+		ensureMsgf(Settings && Settings->Profiles.IsValidIndex(PerProjectUserSettings->AssetViewerProfileIndex), TEXT("Invalid default settings pointer or current profile index"));
+		
+		return !Settings->Profiles[ProfileIndex].bShowEnvironment;
 	}
 
 protected:
@@ -552,27 +558,7 @@ TSharedRef<SWidget> SAnimViewportToolBar::GenerateShowMenu() const
 void SAnimViewportToolBar::FillShowSceneMenu(FMenuBuilder& MenuBuilder) const
 {
 	const FAnimViewportShowCommands& Actions = FAnimViewportShowCommands::Get();
-
-	MenuBuilder.BeginSection("AnimViewportAccessory", LOCTEXT("Viewport_AccessoryLabel", "Accessory"));
-	{
-		MenuBuilder.AddMenuEntry(Actions.ToggleFloor);
-		MenuBuilder.AddMenuEntry(Actions.ToggleSky);
-		MenuBuilder.AddMenuEntry(Actions.AutoAlignFloorToMesh);
-	}
-	MenuBuilder.EndSection();
-
-	MenuBuilder.BeginSection("AnimViewportFloorOffset", LOCTEXT("Viewport_FloorOffsetLabel", "Floor Height Offset"));
-	{
-		TSharedPtr<SWidget> FloorOffsetWidget = SNew(SNumericEntryBox<float>)
-			.Font(FEditorStyle::GetFontStyle(TEXT("MenuItem.Font")))
-			.Value(this, &SAnimViewportToolBar::OnGetFloorOffset)
-			.OnValueChanged(this, &SAnimViewportToolBar::OnFloorOffsetChanged)
-			.ToolTipText(LOCTEXT("FloorOffsetToolTip", "Height offset for the floor mesh (stored per-mesh)"));
-
-		MenuBuilder.AddWidget(FloorOffsetWidget.ToSharedRef(), FText());
-	}
-	MenuBuilder.EndSection();
-
+	
 	MenuBuilder.BeginSection("AnimViewportGrid", LOCTEXT("Viewport_GridLabel", "Grid"));
 	{
 		MenuBuilder.AddMenuEntry(Actions.ToggleGrid);
@@ -835,13 +821,18 @@ FText SAnimViewportToolBar::GetTurnTableMenuLabel() const
 	return Label;
 }
 
-
 FSlateColor SAnimViewportToolBar::GetFontColor() const
 {
+	const UAssetViewerSettings* Settings = UAssetViewerSettings::Get();
+	const UEditorPerProjectUserSettings* PerProjectUserSettings = GetDefault<UEditorPerProjectUserSettings>();
+	const int32 ProfileIndex = Settings->Profiles.IsValidIndex(PerProjectUserSettings->AssetViewerProfileIndex) ? PerProjectUserSettings->AssetViewerProfileIndex : 0;
+
+	ensureMsgf(Settings && Settings->Profiles.IsValidIndex(PerProjectUserSettings->AssetViewerProfileIndex), TEXT("Invalid default settings pointer or current profile index"));
+
 	FLinearColor FontColor;
-	if (Viewport.Pin()->IsShowingSky())
+	if (Settings->Profiles[ProfileIndex].bShowEnvironment)
 	{
-		FontColor = FLinearColor::Black;
+		FontColor = FLinearColor::White;
 	}
 	else
 	{
diff --git a/Engine/Source/Editor/Persona/Private/SAnimationEditorViewport.cpp b/Engine/Source/Editor/Persona/Private/SAnimationEditorViewport.cpp
--- a/Engine/Source/Editor/Persona/Private/SAnimationEditorViewport.cpp
+++ b/Engine/Source/Editor/Persona/Private/SAnimationEditorViewport.cpp
@@ -717,25 +717,7 @@ void SAnimationEditorViewportTabBody::BindCommands()
 		FExecuteAction::CreateSP(this, &SAnimationEditorViewportTabBody::OnShowGrid),
 		FCanExecuteAction(),
 		FIsActionChecked::CreateSP(this, &SAnimationEditorViewportTabBody::IsShowingGrid));
-
-	CommandList.MapAction( 
-		ViewportShowMenuCommands.ToggleFloor,
-		FExecuteAction::CreateSP(this, &SAnimationEditorViewportTabBody::OnShowFloor),
-		FCanExecuteAction(),
-		FIsActionChecked::CreateSP(this, &SAnimationEditorViewportTabBody::IsShowingFloor));
-
-	CommandList.MapAction(
-		ViewportShowMenuCommands.AutoAlignFloorToMesh,
-		FExecuteAction::CreateSP(this, &SAnimationEditorViewportTabBody::OnToggleAutoAlignFloor),
-		FCanExecuteAction(),
-		FIsActionChecked::CreateSP(this, &SAnimationEditorViewportTabBody::IsAutoAlignFloor));
-
-	CommandList.MapAction( 
-		ViewportShowMenuCommands.ToggleSky,
-		FExecuteAction::CreateSP(this, &SAnimationEditorViewportTabBody::OnShowSky),
-		FCanExecuteAction(),
-		FIsActionChecked::CreateSP(this, &SAnimationEditorViewportTabBody::IsShowingSky));
-
+	
 	//Bind LOD preview menu commands
 	const FAnimViewportPlaybackCommands& ViewportPlaybackCommands = FAnimViewportPlaybackCommands::Get();
 
@@ -1058,43 +1040,6 @@ bool SAnimationEditorViewportTabBody::IsShowingGrid() const
 	return AnimViewportClient->IsShowingGrid();
 }
 
-void SAnimationEditorViewportTabBody::OnShowFloor()
-{
-	TSharedRef<FAnimationViewportClient> AnimViewportClient = StaticCastSharedRef<FAnimationViewportClient>(LevelViewportClient.ToSharedRef());
-	AnimViewportClient->OnToggleShowFloor();
-}
-
-bool SAnimationEditorViewportTabBody::IsShowingFloor() const
-{
-	TSharedRef<FAnimationViewportClient> AnimViewportClient = StaticCastSharedRef<FAnimationViewportClient>(LevelViewportClient.ToSharedRef());
-	return AnimViewportClient->IsShowingFloor();
-}
-
-void SAnimationEditorViewportTabBody::OnToggleAutoAlignFloor()
-{
-	TSharedRef<FAnimationViewportClient> AnimViewportClient = StaticCastSharedRef<FAnimationViewportClient>(LevelViewportClient.ToSharedRef());
-	AnimViewportClient->OnToggleAutoAlignFloor();
-}
-
-bool SAnimationEditorViewportTabBody::IsAutoAlignFloor() const
-{
-	TSharedRef<FAnimationViewportClient> AnimViewportClient = StaticCastSharedRef<FAnimationViewportClient>(LevelViewportClient.ToSharedRef());
-	return AnimViewportClient->IsAutoAlignFloor();
-}
-
-void SAnimationEditorViewportTabBody::OnShowSky()
-{
-	TSharedRef<FAnimationViewportClient> AnimViewportClient = StaticCastSharedRef<FAnimationViewportClient>(LevelViewportClient.ToSharedRef());
-	AnimViewportClient->OnToggleShowSky();
-}
-
-bool SAnimationEditorViewportTabBody::IsShowingSky() const
-{
-	TSharedRef<FAnimationViewportClient> AnimViewportClient = StaticCastSharedRef<FAnimationViewportClient>(LevelViewportClient.ToSharedRef());
-	return AnimViewportClient->IsShowingSky();
-}
-
-
 /** Function to set the current playback speed*/
 void SAnimationEditorViewportTabBody::OnSetPlaybackSpeed(int32 PlaybackSpeedMode)
 {
diff --git a/Engine/Source/Editor/Persona/Private/SAnimationEditorViewport.h b/Engine/Source/Editor/Persona/Private/SAnimationEditorViewport.h
--- a/Engine/Source/Editor/Persona/Private/SAnimationEditorViewport.h
+++ b/Engine/Source/Editor/Persona/Private/SAnimationEditorViewport.h
@@ -143,16 +143,7 @@ public:
 
 	/** Can we use gizmos? */
 	bool CanUseGizmos() const;
-
-	/** Function to check whether sky is displayed or not */
-	bool IsShowingSky() const;
-
-	/** Function to check whether floor is displayed or not */
-	bool IsShowingFloor() const;
-
-	/** Function to check whether floor is auto aligned or not */
-	bool IsAutoAlignFloor() const;
-
+	
 	/** Clears our reference to Persona, also cleaning up anything that depends on Persona first */
 	void CleanupPersonaReferences();
 
@@ -266,15 +257,6 @@ private:
 	/** Function to show/hide grid in the viewport */
 	void OnShowGrid();
 
-	/** Function to show/hide floor in the viewport */
-	void OnShowFloor();
-
-	/** Function to enable/disable floor auto align */
-	void OnToggleAutoAlignFloor();
-
-	/** Function to show/hide floor in the viewport */
-	void OnShowSky();
-
 	/** update reference pose with current preview mesh */
 	void UpdateReferencePose();
 
diff --git a/Engine/Source/Editor/Persona/Private/SSkeletonTree.cpp b/Engine/Source/Editor/Persona/Private/SSkeletonTree.cpp
--- a/Engine/Source/Editor/Persona/Private/SSkeletonTree.cpp
+++ b/Engine/Source/Editor/Persona/Private/SSkeletonTree.cpp
@@ -50,23 +50,83 @@ DECLARE_DELEGATE_RetVal_TwoParams(FReply, FOnDraggingBoneItem, const FGeometry&,
 class FSocketTextObjectFactory : public FCustomizableTextObjectFactory
 {
 public:
-	FSocketTextObjectFactory( USkeletalMeshSocket** InDestinationSocket )
+	FSocketTextObjectFactory( USkeleton* InSkeleton, USkeletalMesh* InSkeletalMesh, FName InPasteBone)
 	: FCustomizableTextObjectFactory( GWarn )
-	, DestinationSocket( InDestinationSocket )
+	, PasteBone(InPasteBone)
+	, Skeleton(InSkeleton)
+	, SkeletalMesh(InSkeletalMesh)
 	{
-		
+		check(Skeleton);
 	};
 
+	// Pointer back to the outside world that will hold the final imported socket
+	TArray<USkeletalMeshSocket*> CreatedSockets;
+
 private:
+	
 	virtual bool CanCreateClass(UClass* ObjectClass, bool& bOmitSubObjs) const override { return true; }
 
 	virtual void ProcessConstructedObject( UObject* CreatedObject )
 	{
-		*DestinationSocket = (USkeletalMeshSocket*)CreatedObject;
+		USkeletalMeshSocket* NewSocket = CastChecked<USkeletalMeshSocket>(CreatedObject);
+		CreatedSockets.Add(NewSocket);
+
+		const FReferenceSkeleton* RefSkel = nullptr;
+
+		if (USkeletalMesh* SocketMesh = Cast<USkeletalMesh>(NewSocket->GetOuter()))
+		{
+			SocketMesh->GetMeshOnlySocketList().Add(NewSocket);
+			RefSkel = &SocketMesh->RefSkeleton;
+		}
+		else if(USkeleton* SocketSkeleton = Cast<USkeleton>(NewSocket->GetOuter()))
+		{
+			SocketSkeleton->Sockets.Add(NewSocket);
+			RefSkel = &SocketSkeleton->GetReferenceSkeleton();
+		}
+		else
+		{
+			check(false) //Unknown socket outer
+		}
+
+		if (!PasteBone.IsNone())
+		{
+			// Override the bone name to the one we pasted to
+			NewSocket->BoneName = PasteBone;
+		}
+		else
+		{
+			//Validate BoneName
+			if (RefSkel->FindBoneIndex(NewSocket->BoneName) == INDEX_NONE)
+			{
+				NewSocket->BoneName = RefSkel->GetBoneName(0);
+			}
+		}
+	}
+
+	virtual void ProcessUnidentifiedLine(const FString& StrLine)
+	{
+		bool bIsOnSkeleton;
+		FParse::Bool(*StrLine, TEXT("IsOnSkeleton="), bIsOnSkeleton);
+		bExpectingMeshSocket = !bIsOnSkeleton;
 	}
 
-	// Pointer back to the outside world that will hold the final imported socket
-	USkeletalMeshSocket** DestinationSocket;
+	virtual UObject* GetParentForNewObject(const UClass* ObjClass) 
+	{ 
+		UObject* Target = (bExpectingMeshSocket && SkeletalMesh)? (UObject*)SkeletalMesh : (UObject*)Skeleton;
+		Target->Modify();
+		return Target;
+	}
+
+	const FName PasteBone;
+
+	// Track what type of socket we will be processing next
+	bool bExpectingMeshSocket;
+
+	// Target for Skeleton Sockets
+	USkeleton* Skeleton; 
+
+	// Target for Mesh Sockets (could be nullptr)
+	USkeletalMesh* SkeletalMesh;
 };
 
 //////////////////////////////////////////////////////////////////////////
@@ -1152,7 +1212,11 @@ void SSkeletonTree::BindCommands()
 
 	CommandList.MapAction(
 		MenuActions.PasteSockets,
-		FExecuteAction::CreateSP( this, &SSkeletonTree::OnPasteSockets ) );
+		FExecuteAction::CreateSP( this, &SSkeletonTree::OnPasteSockets, false ) );
+
+	CommandList.MapAction(
+		MenuActions.PasteSocketsToSelectedBone,
+		FExecuteAction::CreateSP(this, &SSkeletonTree::OnPasteSockets, true));
 }
 
 TSharedRef<ITableRow> SSkeletonTree::MakeTreeRowWidget(TSharedPtr<FDisplayedTreeRowInfo> InInfo, const TSharedRef<STableViewBase>& OwnerTable)
@@ -1549,6 +1613,7 @@ TSharedPtr< SWidget > SSkeletonTree::CreateContextMenu()
 			{
 				MenuBuilder.AddMenuEntry( Actions.AddSocket );
 				MenuBuilder.AddMenuEntry( Actions.PasteSockets );
+				MenuBuilder.AddMenuEntry( Actions.PasteSocketsToSelectedBone );
 			}
 
 			MenuBuilder.EndSection();
@@ -1886,14 +1951,14 @@ FString SSkeletonTree::SerializeSocketToString( USkeletalMeshSocket* Socket, con
 	return SocketString;
 }
 
-void SSkeletonTree::OnPasteSockets()
+void SSkeletonTree::OnPasteSockets(bool bPasteToSelectedBone)
 {
 	FBoneTreeSelection TreeSelection(SkeletonTreeView->GetSelectedItems());
 
 	// Pasting sockets should only work if there is just one bone selected
 	if ( TreeSelection.IsSingleOfTypeSelected(ESkeletonTreeRowType::Bone) )
 	{
-		FName DestBoneName = *static_cast<FName*>( TreeSelection.GetSingleSelectedItem()->GetData() );
+		FName DestBoneName = bPasteToSelectedBone ? *static_cast<FName*>( TreeSelection.GetSingleSelectedItem()->GetData() ) : NAME_None;
 
 		FString PasteString;
 		FPlatformMisc::ClipboardPaste( PasteString );
@@ -1909,51 +1974,14 @@ void SSkeletonTree::OnPasteSockets()
 			int32 NumSocketsToPaste;
 			FParse::Line( &PastePtr, PasteLine );	// Need this to advance PastePtr, for multiple sockets
 			FParse::Value( *PasteLine, TEXT( "NumSockets=" ), NumSocketsToPaste );
-			FParse::Line( &PastePtr, PasteLine );
-
-			for ( int32 i = 0; i < NumSocketsToPaste; ++i )
-			{					
-				bool bIsOnSkeleton;
-				FParse::Bool( *PasteLine, TEXT( "IsOnSkeleton=" ), bIsOnSkeleton );
-
-				USkeletalMeshSocket* NewSocket;
-				FSocketTextObjectFactory TextObjectFactory( &NewSocket );
-
-				if ( bIsOnSkeleton )
-				{
-					TargetSkeleton->Modify();
 
-					TextObjectFactory.ProcessBuffer( TargetSkeleton, RF_Transactional, PastePtr );
-					check(NewSocket);
-				}
-				else
-				{
-					USkeletalMesh* Mesh = PersonaPtr.Pin()->GetMesh();
-
-					Mesh->Modify();
-
-					TextObjectFactory.ProcessBuffer( Mesh, RF_Transactional, PastePtr );
-					check(NewSocket);
-				}
-
-				// Override the bone name to the one we pasted to
-				NewSocket->BoneName = DestBoneName;
+			FSocketTextObjectFactory TextObjectFactory(TargetSkeleton, PersonaPtr.Pin()->GetMesh(), DestBoneName);
+			TextObjectFactory.ProcessBuffer(nullptr, RF_Transactional, PastePtr);
 
+			for (USkeletalMeshSocket* NewSocket : TextObjectFactory.CreatedSockets)
+			{
 				// Check the socket name is unique
-				NewSocket->SocketName = PersonaPtr.Pin()->GenerateUniqueSocketName( NewSocket->SocketName );
-
-				// Skip ahead in the stream to the next socket (if there is one)
-				PastePtr = FCString::Strstr( PastePtr, TEXT( "IsOnSkeleton=" ) );
-
-				if ( bIsOnSkeleton )
-				{
-					TargetSkeleton->Sockets.Add( NewSocket );
-				}
-				else
-				{
-					USkeletalMesh* Mesh = PersonaPtr.Pin()->GetMesh();
-					Mesh->GetMeshOnlySocketList().Add( NewSocket );
-				}
+				NewSocket->SocketName = PersonaPtr.Pin()->GenerateUniqueSocketName(NewSocket->SocketName);
 			}
 		}
 		CreateFromSkeleton( TargetSkeleton->GetBoneTree() );
diff --git a/Engine/Source/Editor/Persona/Private/SSkeletonTree.h b/Engine/Source/Editor/Persona/Private/SSkeletonTree.h
--- a/Engine/Source/Editor/Persona/Private/SSkeletonTree.h
+++ b/Engine/Source/Editor/Persona/Private/SSkeletonTree.h
@@ -494,7 +494,7 @@ private:
 	FString SerializeSocketToString( USkeletalMeshSocket* Socket, const FDisplayedSocketInfo* DisplayedSocketInfo ) const;
 
 	/** Function to paste selected sockets from the clipboard */
-	void OnPasteSockets();
+	void OnPasteSockets(bool bPasteToSelectedBone);
 
 	/** Function to add a socket to the selected bone (skeleton, not mesh) */
 	void OnAddSocket();
diff --git a/Engine/Source/Editor/Persona/Private/Shared/PersonaMode.cpp b/Engine/Source/Editor/Persona/Private/Shared/PersonaMode.cpp
--- a/Engine/Source/Editor/Persona/Private/Shared/PersonaMode.cpp
+++ b/Engine/Source/Editor/Persona/Private/Shared/PersonaMode.cpp
@@ -9,6 +9,8 @@
 #include "SSkeletonSlotNames.h"
 #include "SSkeletonSmartNameManager.h"
 #include "SSkeletonBlendProfiles.h"
+#include "SDocktab.h"
+#include "SAdvancedPreviewDetailsTab.h"
 
 #define LOCTEXT_NAMESPACE "PersonaModes"
 
@@ -44,6 +46,8 @@ const FName FPersonaTabs::SkeletonSlotGroupNamesID("SkeletonSlotGroupNames");
 const FName FPersonaTabs::CurveNameManagerID("CurveNameManager");
 const FName FPersonaTabs::BlendProfileManagerID("BlendProfileManager");
 
+const FName FPersonaTabs::AdvancedPreviewSceneSettingsID("AdvancedPreviewTab");
+
 /////////////////////////////////////////////////////
 // FPersonaMode
 
@@ -72,16 +76,17 @@ FPersonaAppMode::FPersonaAppMode(TSharedPtr<class FPersona> InPersona, FName InM
 	PersonaTabFactories.RegisterFactory(MakeShareable(new FSkeletonSlotNamesSummoner(InPersona)));
 	PersonaTabFactories.RegisterFactory(MakeShareable(new FSkeletonCurveNameManagerSummoner(InPersona)));
 	PersonaTabFactories.RegisterFactory(MakeShareable(new FSkeletonBlendProfilesSummoner(InPersona)));
+	PersonaTabFactories.RegisterFactory(MakeShareable(new FAdvancedPreviewSceneTabSummoner(InPersona)));
 }
 
 void FPersonaAppMode::RegisterTabFactories(TSharedPtr<FTabManager> InTabManager)
 {
-	TSharedPtr<FBlueprintEditor> BP = MyPersona.Pin();
+	TSharedPtr<FPersona> Persona = MyPersona.Pin();
 	
-	BP->RegisterToolbarTab(InTabManager.ToSharedRef());
-
 	// Mode-specific setup
-	BP->PushTabFactories(PersonaTabFactories);
+	Persona->PushTabFactories(PersonaTabFactories);
+
+	Persona->RegisterToolbarTab(InTabManager.ToSharedRef());
 }
 
 void FPersonaAppMode::PostActivateMode()
@@ -432,4 +437,30 @@ FText FAnimBlueprintParentPlayerEditorSummoner::GetTabToolTipText(const FWorkflo
 	return LOCTEXT("AnimSubClassTabToolTip", "Editor for overriding the animation assets referenced by the parent animation graph.");
 }
 
+/////////////////////////////////////////////////////
+// FAdvancedPreviewSceneTabSummoner
+
+FAdvancedPreviewSceneTabSummoner::FAdvancedPreviewSceneTabSummoner(TSharedPtr<class FAssetEditorToolkit> InHostingApp)
+	: FWorkflowTabFactory(FPersonaTabs::AdvancedPreviewSceneSettingsID, InHostingApp)
+{
+	TabLabel = LOCTEXT("PreviewSceneSettingsTab", "Preview Scene Settings");
+	TabIcon = FSlateIcon(FEditorStyle::GetStyleSetName(), "LevelEditor.Tabs.Details");	
+	bIsSingleton = true;
+	
+	ViewMenuDescription = LOCTEXT("AdvancedPreviewScene", "Preview Scene Settings");
+	ViewMenuTooltip = LOCTEXT("AdvancedPreviewScene_ToolTip", "Shows the advanced preview scene settings");
+}
+
+TSharedRef<SWidget> FAdvancedPreviewSceneTabSummoner::CreateTabBody(const FWorkflowTabSpawnInfo& Info) const
+{
+	TSharedPtr<FPersona> PersonaPtr = StaticCastSharedPtr<FPersona>(HostingApp.Pin());	
+	return SNew(SAdvancedPreviewDetailsTab).PreviewScenePtr(&(PersonaPtr->GetPreviewScene()));
+}
+
+FText FAdvancedPreviewSceneTabSummoner::GetTabToolTipText(const FWorkflowTabSpawnInfo& Info) const
+{
+	return LOCTEXT("AdvancedPreviewSettingsToolTip", "The Advanced Preview Settings tab will let you alter the preview scene's settings.");
+}
+
 #undef LOCTEXT_NAMESPACE
+
diff --git a/Engine/Source/Editor/Persona/Private/Shared/PersonaMode.h b/Engine/Source/Editor/Persona/Private/Shared/PersonaMode.h
--- a/Engine/Source/Editor/Persona/Private/Shared/PersonaMode.h
+++ b/Engine/Source/Editor/Persona/Private/Shared/PersonaMode.h
@@ -41,6 +41,9 @@ struct FPersonaTabs
 	static const FName CurveNameManagerID;
 	static const FName BlendProfileManagerID;
 
+	// Advanced Preview Scene
+	static const FName AdvancedPreviewSceneSettingsID;
+
 	// Blueprint Document
 
 	// Inherited from blueprint editor
@@ -284,4 +287,15 @@ private:
 
 };
 
+/////////////////////////////////////////////////////
+// FAdvancedPreviewSceneTabSummoner
+
+struct FAdvancedPreviewSceneTabSummoner : public FWorkflowTabFactory
+{
+public:
+ 	FAdvancedPreviewSceneTabSummoner(TSharedPtr<class FAssetEditorToolkit> InHostingApp);
+	virtual TSharedRef<SWidget> CreateTabBody(const FWorkflowTabSpawnInfo& Info) const override;
+	virtual FText GetTabToolTipText(const FWorkflowTabSpawnInfo& Info) const override;
+};
+
 #undef LOCTEXT_NAMESPACE
\ No newline at end of file
diff --git a/Engine/Source/Editor/Persona/Private/SkeletonTreeCommands.cpp b/Engine/Source/Editor/Persona/Private/SkeletonTreeCommands.cpp
--- a/Engine/Source/Editor/Persona/Private/SkeletonTreeCommands.cpp
+++ b/Engine/Source/Editor/Persona/Private/SkeletonTreeCommands.cpp
@@ -16,7 +16,8 @@ void FSkeletonTreeCommands::RegisterCommands()
 	UI_COMMAND( CopyBoneNames, "Copy Selected Bone Names", "Copy selected bone names to clipboard", EUserInterfaceActionType::Button, FInputChord() );
 	UI_COMMAND( ResetBoneTransforms, "Reset Selected Bone Transforms", "Reset the transforms of the selected bones", EUserInterfaceActionType::Button, FInputChord() );
 	UI_COMMAND( CopySockets, "Copy Selected Sockets", "Copy selected sockets to clipboard", EUserInterfaceActionType::Button, FInputChord( EModifierKey::Control, EKeys::C ) );
-	UI_COMMAND( PasteSockets, "Paste Selected Sockets", "Paste selected sockets from clipboard", EUserInterfaceActionType::Button, FInputChord( EModifierKey::Control, EKeys::V ) );
+	UI_COMMAND( PasteSockets, "Paste Sockets", "Paste sockets from clipboard (try to add them to the original copy bone or root if that bone is not available)", EUserInterfaceActionType::Button, FInputChord( EModifierKey::Control, EKeys::V ) );
+	UI_COMMAND( PasteSocketsToSelectedBone, "Paste Sockets To Selected Bone", "Paste sockets from clipboard to selected bone", EUserInterfaceActionType::Button, FInputChord(EModifierKey::Control | EModifierKey::Shift, EKeys::V));
 
 	UI_COMMAND( AddSocket, "Add Socket", "Add a socket to this bone in the skeleton (disabled when socket filter is set to \"Mesh Sockets\" or \"Sockets Hidden\" mode)", EUserInterfaceActionType::Button, FInputChord() );
 	UI_COMMAND( CreateMeshSocket, "Create Mesh Socket", "Duplicate this socket from skeleton to the current mesh and modify the socket data for it", EUserInterfaceActionType::Button, FInputChord() );
diff --git a/Engine/Source/Editor/Persona/Private/SkeletonTreeCommands.h b/Engine/Source/Editor/Persona/Private/SkeletonTreeCommands.h
--- a/Engine/Source/Editor/Persona/Private/SkeletonTreeCommands.h
+++ b/Engine/Source/Editor/Persona/Private/SkeletonTreeCommands.h
@@ -80,6 +80,9 @@ public:
 
 	/** Paste sockets from clipboard */
 	TSharedPtr< FUICommandInfo > PasteSockets;
+
+	/** Paste sockets from clipboard */
+	TSharedPtr< FUICommandInfo > PasteSocketsToSelectedBone;
 };
 
 #endif // __SkeletonTreeCommands_h_
diff --git a/Engine/Source/Editor/PhAT/Private/PhAT.cpp b/Engine/Source/Editor/PhAT/Private/PhAT.cpp
--- a/Engine/Source/Editor/PhAT/Private/PhAT.cpp
+++ b/Engine/Source/Editor/PhAT/Private/PhAT.cpp
@@ -29,6 +29,7 @@
 #include "SDockTab.h"
 #include "PhysicsEngine/BodySetup.h"
 #include "PhysicsEngine/PhysicsConstraintTemplate.h"
+#include "PhysicsEngine/ConstraintUtils.h"
 #include "Engine/Selection.h"
 #include "Engine/StaticMesh.h"
 #include "EngineLogs.h"
@@ -1286,18 +1287,6 @@ void FPhAT::BindCommands()
 		Commands.DeleteAllBodiesBelow,
 		FExecuteAction::CreateSP(this, &FPhAT::OnDeleteAllBodiesBelow));
 
-	ToolkitCommands->MapAction(
-		Commands.ToggleMotor,
-		FExecuteAction::CreateSP(this, &FPhAT::OnToggleMotor));
-
-	ToolkitCommands->MapAction(
-		Commands.EnableMotorsBelow,
-		FExecuteAction::CreateSP(this, &FPhAT::OnEnableMotorsBelow));
-
-	ToolkitCommands->MapAction(
-		Commands.DisableMotorsBelow,
-		FExecuteAction::CreateSP(this, &FPhAT::OnDisableMotorsBelow));
-
 	ToolkitCommands->MapAction(
 		Commands.SelectionLock,
 		FExecuteAction::CreateSP(this, &FPhAT::OnLockSelection),
@@ -1823,20 +1812,20 @@ void FPhAT::CreateOrConvertConstraint(EPhATConstraintType ConstraintType)
 
 		if(ConstraintType == EPCT_BSJoint)
 		{
-			ConstraintSetup->DefaultInstance.ConfigureAsBS();
+			ConstraintUtils::ConfigureAsBallAndSocket(ConstraintSetup->DefaultInstance);
 		}
 		else if(ConstraintType == EPCT_Hinge)
 		{
-			ConstraintSetup->DefaultInstance.ConfigureAsHinge();
+			ConstraintUtils::ConfigureAsHinge(ConstraintSetup->DefaultInstance);
 		}
 		else if(ConstraintType == EPCT_Prismatic)
 		{
-			ConstraintSetup->DefaultInstance.ConfigureAsPrismatic();
+			ConstraintUtils::ConfigureAsPrismatic(ConstraintSetup->DefaultInstance);
 		}
 		else if(ConstraintType == EPCT_SkelJoint)
 		{
-			ConstraintSetup->DefaultInstance.ConfigureAsSkelJoint();
-		}	
+			ConstraintUtils::ConfigureAsSkelJoint(ConstraintSetup->DefaultInstance);
+		}
 	}
 
 	RefreshHierachyTree();
@@ -1844,43 +1833,6 @@ void FPhAT::CreateOrConvertConstraint(EPhATConstraintType ConstraintType)
 	
 }
 
-void FPhAT::SetConstraintsBelowSelectedMotorised(bool bMotorised)
-{
-	SharedData->PhysicsAsset->Modify();
-
-	for(int32 i=0; i<SharedData->SelectedConstraints.Num(); ++i)
-	{
-		// Get the index of this constraint
-		UPhysicsConstraintTemplate* BaseSetup = SharedData->PhysicsAsset->ConstraintSetup[SharedData->SelectedConstraints[i].Index];
-
-		TArray<int32> BelowConstraints;
-		int32 BaseIndex = SharedData->EditorSkelMesh->RefSkeleton.FindBoneIndex(BaseSetup->DefaultInstance.JointName);
-
-		// Iterate over all other joints, looking for 'children' of this one
-		for (int32 j = 0; j < SharedData->PhysicsAsset->ConstraintSetup.Num(); ++j)
-		{
-			UPhysicsConstraintTemplate* ConstraintSetup = SharedData->PhysicsAsset->ConstraintSetup[j];
-			FName TestName = ConstraintSetup->DefaultInstance.JointName;
-			int32 TestIndex = SharedData->EditorSkelMesh->RefSkeleton.FindBoneIndex(TestName);
-
-			// We want to return this constraint as well.
-			if (TestIndex == BaseIndex || SharedData->EditorSkelMesh->RefSkeleton.BoneIsChildOf(TestIndex, BaseIndex))
-			{
-				BelowConstraints.Add(j);
-			}
-		}
-
-		for (int32 j = 0; j < BelowConstraints.Num(); ++j)
-		{
-			int32 ConIndex = BelowConstraints[j];
-			FConstraintInstance* ConstraintInstance = &SharedData->PhysicsAsset->ConstraintSetup[ConIndex]->DefaultInstance;
-
-			ConstraintInstance->bAngularOrientationDrive = bMotorised;
-		}
-
-	}
-}
-
 void FPhAT::AddNewPrimitive(EKCollisionPrimitiveType InPrimitiveType, bool bCopySelected)
 {
 	check(!bCopySelected || SharedData->SelectedBodies.Num() == 1);	//we only support this for one selection
@@ -2885,27 +2837,6 @@ void FPhAT::OnDeleteAllBodiesBelow()
 	
 }
 
-void FPhAT::OnToggleMotor()
-{
-	for(int32 i=0; i<SharedData->SelectedConstraints.Num(); ++i)
-	{
-		UPhysicsConstraintTemplate* ConSetup = SharedData->PhysicsAsset->ConstraintSetup[SharedData->SelectedConstraints[i].Index];
-		FConstraintInstance* CI = &ConSetup->DefaultInstance;
-
-		CI->bAngularOrientationDrive = !CI->bAngularOrientationDrive;
-	}
-}
-
-void FPhAT::OnEnableMotorsBelow()
-{
-	SetConstraintsBelowSelectedMotorised(true);
-}
-
-void FPhAT::OnDisableMotorsBelow()
-{
-	SetConstraintsBelowSelectedMotorised(false);
-}
-
 void FPhAT::OnLockSelection()
 {
 	SharedData->bSelectionLock = !SharedData->bSelectionLock;
diff --git a/Engine/Source/Editor/PhAT/Private/PhAT.h b/Engine/Source/Editor/PhAT/Private/PhAT.h
--- a/Engine/Source/Editor/PhAT/Private/PhAT.h
+++ b/Engine/Source/Editor/PhAT/Private/PhAT.h
@@ -147,8 +147,7 @@ private:
 	/** Constraint editing helper methods */
 	void SnapConstraintToBone(int32 ConstraintIndex, const FTransform& WParentFrame);
 	void CreateOrConvertConstraint(EPhATConstraintType ConstraintType);
-	void SetConstraintsBelowSelectedMotorised(bool bMotorised);
-
+	
 	/** Collision editing helper methods */
 	void AddNewPrimitive(EKCollisionPrimitiveType PrimitiveType, bool bCopySelected = false);
 	void SetBodiesBelowSelectedPhysicsType( EPhysicsType InPhysicsType, bool bMarkAsDirty);
@@ -227,9 +226,6 @@ private:
 	bool IsBodyPhysicsType( EPhysicsType InPhysicsType );
 	void OnDeleteBody();
 	void OnDeleteAllBodiesBelow();
-	void OnToggleMotor();
-	void OnEnableMotorsBelow();
-	void OnDisableMotorsBelow();
 	void OnLockSelection();
 	void OnDeleteSelection();
 	void OnCycleConstraintOrientation();
diff --git a/Engine/Source/Editor/PhAT/Private/PhATActions.cpp b/Engine/Source/Editor/PhAT/Private/PhATActions.cpp
--- a/Engine/Source/Editor/PhAT/Private/PhATActions.cpp
+++ b/Engine/Source/Editor/PhAT/Private/PhATActions.cpp
@@ -87,8 +87,8 @@ void FPhATCommands::RegisterCommands()
 	UI_COMMAND(CycleConstraintOrientation, "Cycle selected constraint orientation", "", EUserInterfaceActionType::Button, FInputChord(EKeys::Q));
 	UI_COMMAND(CycleConstraintActive, "Cycle active constraint", "", EUserInterfaceActionType::Button, FInputChord(EKeys::Four));
 	UI_COMMAND(ToggleSwing1, "Toggle Swing1 Constraint", "", EUserInterfaceActionType::Button, FInputChord(EKeys::One));
-	UI_COMMAND(ToggleSwing2, "Toggle Swing2 Constraint", "", EUserInterfaceActionType::Button, FInputChord(EKeys::Three));
-	UI_COMMAND(ToggleTwist, "Toggle Twist Constraint", "", EUserInterfaceActionType::Button, FInputChord(EKeys::Two));
+	UI_COMMAND(ToggleSwing2, "Toggle Swing2 Constraint", "", EUserInterfaceActionType::Button, FInputChord(EKeys::Two));
+	UI_COMMAND(ToggleTwist, "Toggle Twist Constraint", "", EUserInterfaceActionType::Button, FInputChord(EKeys::Three));
 	UI_COMMAND(FocusOnSelection, "Focus the viewport on the current selection", "", EUserInterfaceActionType::Button, FInputChord(EKeys::F));
 }
 
diff --git a/Engine/Source/Editor/PhAT/Private/PhATSharedData.cpp b/Engine/Source/Editor/PhAT/Private/PhATSharedData.cpp
--- a/Engine/Source/Editor/PhAT/Private/PhATSharedData.cpp
+++ b/Engine/Source/Editor/PhAT/Private/PhATSharedData.cpp
@@ -1134,21 +1134,21 @@ void FPhATSharedData::CycleCurrentConstraintActive()
 		UPhysicsConstraintTemplate* ConstraintTemplate = PhysicsAsset->ConstraintSetup[GetSelectedConstraint()->Index];
 		FConstraintInstance & DefaultInstance = ConstraintTemplate->DefaultInstance;
 
-		if(DefaultInstance.AngularSwing1Motion != ACM_Limited && DefaultInstance.AngularSwing2Motion != ACM_Limited)
+		if(DefaultInstance.GetAngularSwing1Motion() != ACM_Limited && DefaultInstance.GetAngularSwing2Motion() != ACM_Limited)
 		{
-			DefaultInstance.AngularSwing1Motion = ACM_Limited;
-			DefaultInstance.AngularSwing2Motion = ACM_Locked;
-			DefaultInstance.AngularTwistMotion = ACM_Locked;
-		}else if(DefaultInstance.AngularSwing2Motion != ACM_Limited && DefaultInstance.AngularTwistMotion != ACM_Limited)
+			DefaultInstance.SetAngularSwing1Motion(ACM_Limited);
+			DefaultInstance.SetAngularSwing2Motion(ACM_Locked);
+			DefaultInstance.SetAngularTwistMotion(ACM_Locked);
+		}else if(DefaultInstance.GetAngularSwing2Motion() != ACM_Limited && DefaultInstance.GetAngularTwistMotion() != ACM_Limited)
 		{
-			DefaultInstance.AngularSwing1Motion = ACM_Locked;
-			DefaultInstance.AngularSwing2Motion = ACM_Limited;
-			DefaultInstance.AngularTwistMotion = ACM_Locked;
+			DefaultInstance.SetAngularSwing1Motion(ACM_Locked);
+			DefaultInstance.SetAngularSwing2Motion(ACM_Limited);
+			DefaultInstance.SetAngularTwistMotion(ACM_Locked);
 		}else
 		{
-			DefaultInstance.AngularSwing1Motion = ACM_Locked;
-			DefaultInstance.AngularSwing2Motion = ACM_Locked;
-			DefaultInstance.AngularTwistMotion = ACM_Limited;
+			DefaultInstance.SetAngularSwing1Motion(ACM_Locked);
+			DefaultInstance.SetAngularSwing2Motion(ACM_Locked);
+			DefaultInstance.SetAngularTwistMotion(ACM_Limited);
 		}
 		
 	}
@@ -1163,13 +1163,13 @@ void FPhATSharedData::ToggleConstraint(EPhATConstraintType Constraint)
 
 		if(Constraint == PCT_Swing1)
 		{
-			DefaultInstance.AngularSwing1Motion = DefaultInstance.AngularSwing1Motion == ACM_Limited ? ACM_Locked : ACM_Limited;
+			DefaultInstance.SetAngularSwing1Motion(DefaultInstance.GetAngularSwing1Motion() == ACM_Limited ? ACM_Locked : ACM_Limited);
 		}else if(Constraint == PCT_Swing2)
 		{
-			DefaultInstance.AngularSwing2Motion = DefaultInstance.AngularSwing2Motion == ACM_Limited ? ACM_Locked : ACM_Limited;
+			DefaultInstance.SetAngularSwing2Motion(DefaultInstance.GetAngularSwing2Motion() == ACM_Limited ? ACM_Locked : ACM_Limited);
 		}else
 		{
-			DefaultInstance.AngularTwistMotion = DefaultInstance.AngularTwistMotion == ACM_Limited ? ACM_Locked : ACM_Limited;
+			DefaultInstance.SetAngularTwistMotion(DefaultInstance.GetAngularTwistMotion() == ACM_Limited ? ACM_Locked : ACM_Limited);
 		}
 		
 	}
diff --git a/Engine/Source/Editor/StaticMeshEditor/Private/SStaticMeshEditorViewport.cpp b/Engine/Source/Editor/StaticMeshEditor/Private/SStaticMeshEditorViewport.cpp
--- a/Engine/Source/Editor/StaticMeshEditor/Private/SStaticMeshEditorViewport.cpp
+++ b/Engine/Source/Editor/StaticMeshEditor/Private/SStaticMeshEditorViewport.cpp
@@ -6,7 +6,7 @@
 
 #include "MouseDeltaTracker.h"
 #include "SStaticMeshEditorViewport.h"
-#include "PreviewScene.h"
+#include "AdvancedPreviewScene.h"
 #include "Runtime/Engine/Public/Slate/SceneViewport.h"
 #include "StaticMeshResources.h"
 
@@ -92,6 +92,10 @@ public:
 
 void SStaticMeshEditorViewport::Construct(const FArguments& InArgs)
 {
+	//PreviewScene = new FAdvancedPreviewScene(FPreviewScene::ConstructionValues(), 
+
+	PreviewScene.SetFloorOffset(-InArgs._ObjectToEdit->ExtendedBounds.Origin.Z + InArgs._ObjectToEdit->ExtendedBounds.BoxExtent.Z);
+
 	StaticMeshEditorPtr = InArgs._StaticMeshEditor;
 
 	StaticMesh = InArgs._ObjectToEdit;
@@ -113,6 +117,11 @@ void SStaticMeshEditorViewport::Construct(const FArguments& InArgs)
 		];
 
 	FCoreUObjectDelegates::OnObjectPropertyChanged.AddRaw(this, &SStaticMeshEditorViewport::OnObjectPropertyChanged);
+}
+
+SStaticMeshEditorViewport::SStaticMeshEditorViewport()
+	: PreviewScene(FPreviewScene::ConstructionValues())
+{
 
 }
 
@@ -394,7 +403,7 @@ TSharedRef<FEditorViewportClient> SStaticMeshEditorViewport::MakeEditorViewportC
 
 	EditorViewportClient->bSetListenerPosition = false;
 
-	EditorViewportClient->SetRealtime( false );
+	EditorViewportClient->SetRealtime( true );
 	EditorViewportClient->VisibilityDelegate.BindSP( this, &SStaticMeshEditorViewport::IsVisible );
 
 	return EditorViewportClient.ToSharedRef();
diff --git a/Engine/Source/Editor/StaticMeshEditor/Private/SStaticMeshEditorViewport.h b/Engine/Source/Editor/StaticMeshEditor/Private/SStaticMeshEditorViewport.h
--- a/Engine/Source/Editor/StaticMeshEditor/Private/SStaticMeshEditorViewport.h
+++ b/Engine/Source/Editor/StaticMeshEditor/Private/SStaticMeshEditorViewport.h
@@ -2,7 +2,7 @@
 
 #pragma once
 
-#include "PreviewScene.h"
+#include "AdvancedPreviewScene.h"
 #include "SEditorViewport.h"
 #include "Editor/UnrealEd/Public/SCommonEditorViewportToolbarBase.h"
 
@@ -18,6 +18,7 @@ public:
 	SLATE_END_ARGS()
 
 	void Construct(const FArguments& InArgs);
+	SStaticMeshEditorViewport();
 	~SStaticMeshEditorViewport();
 	
 	// FGCObject interface
@@ -88,6 +89,8 @@ public:
 	virtual void OnFloatingButtonClicked() override;
 	// End of ICommonEditorViewportToolbarInfoProvider interface
 
+	/** Returns the preview scene being renderd in the viewport */
+	FAdvancedPreviewScene& GetPreviewScene() { return PreviewScene; }
 protected:
 	/** SEditorViewport interface */
 	virtual TSharedRef<FEditorViewportClient> MakeEditorViewportClient() override;
@@ -126,7 +129,7 @@ private:
 	TWeakPtr<IStaticMeshEditor> StaticMeshEditorPtr;
 
 	/** The scene for this viewport. */
-	FPreviewScene PreviewScene;
+	FAdvancedPreviewScene PreviewScene;
 
 	/** Editor viewport client */
 	TSharedPtr<class FStaticMeshEditorViewportClient> EditorViewportClient;
diff --git a/Engine/Source/Editor/StaticMeshEditor/Private/StaticMeshEditor.cpp b/Engine/Source/Editor/StaticMeshEditor/Private/StaticMeshEditor.cpp
--- a/Engine/Source/Editor/StaticMeshEditor/Private/StaticMeshEditor.cpp
+++ b/Engine/Source/Editor/StaticMeshEditor/Private/StaticMeshEditor.cpp
@@ -32,6 +32,8 @@
 #include "Engine/Selection.h"
 #include "PhysicsEngine/BodySetup.h"
 
+#include "SAdvancedPreviewDetailsTab.h"
+
 #define LOCTEXT_NAMESPACE "StaticMeshEditor"
 
 DEFINE_LOG_CATEGORY_STATIC(LogStaticMeshEditor, Log, All);
@@ -51,7 +53,7 @@ const FName FStaticMeshEditor::ViewportTabId( TEXT( "StaticMeshEditor_Viewport"
 const FName FStaticMeshEditor::PropertiesTabId( TEXT( "StaticMeshEditor_Properties" ) );
 const FName FStaticMeshEditor::SocketManagerTabId( TEXT( "StaticMeshEditor_SocketManager" ) );
 const FName FStaticMeshEditor::CollisionTabId( TEXT( "StaticMeshEditor_Collision" ) );
-
+const FName FStaticMeshEditor::PreviewSceneSettingsTabId(TEXT("StaticMeshEditor_PreviewScene"));
 
 void FStaticMeshEditor::RegisterTabSpawners(const TSharedRef<class FTabManager>& InTabManager)
 {
@@ -79,6 +81,11 @@ void FStaticMeshEditor::RegisterTabSpawners(const TSharedRef<class FTabManager>&
 		.SetDisplayName( LOCTEXT("CollisionTab", "Convex Decomposition") )
 		.SetGroup(WorkspaceMenuCategoryRef)
 		.SetIcon(FSlateIcon(FEditorStyle::GetStyleSetName(), "StaticMeshEditor.Tabs.ConvexDecomposition"));
+
+	InTabManager->RegisterTabSpawner(PreviewSceneSettingsTabId, FOnSpawnTab::CreateSP(this, &FStaticMeshEditor::SpawnTab_PreviewSceneSettings))
+		.SetDisplayName(LOCTEXT("PreviewSceneTab", "Preview Scene Settings"))
+		.SetGroup(WorkspaceMenuCategoryRef)
+		.SetIcon(FSlateIcon(FEditorStyle::GetStyleSetName(), "LevelEditor.Tabs.Details"));	
 }
 
 void FStaticMeshEditor::UnregisterTabSpawners(const TSharedRef<class FTabManager>& InTabManager)
@@ -89,6 +96,7 @@ void FStaticMeshEditor::UnregisterTabSpawners(const TSharedRef<class FTabManager
 	InTabManager->UnregisterTabSpawner( PropertiesTabId );
 	InTabManager->UnregisterTabSpawner( SocketManagerTabId );
 	InTabManager->UnregisterTabSpawner( CollisionTabId );
+	InTabManager->UnregisterTabSpawner( PreviewSceneSettingsTabId );
 }
 
 
@@ -361,6 +369,16 @@ TSharedRef<SDockTab> FStaticMeshEditor::SpawnTab_Collision( const FSpawnTabArgs&
 		];
 }
 
+TSharedRef<SDockTab> FStaticMeshEditor::SpawnTab_PreviewSceneSettings(const FSpawnTabArgs& Args)
+{
+	check(Args.GetTabId() == PreviewSceneSettingsTabId);
+	return SNew(SDockTab)
+		.Label(LOCTEXT("StaticMeshPreviewScene_TabTitle", "Preview Scene Settings"))
+		[
+			AdvancedPreviewSettingsWidget.ToSharedRef()
+		];
+}
+
 void FStaticMeshEditor::BindCommands()
 {
 	const FStaticMeshEditorCommands& Commands = FStaticMeshEditorCommands::Get();
@@ -548,6 +566,9 @@ void FStaticMeshEditor::BuildSubTools()
 	{
 		LODLevelCombo->SetSelectedItem(LODLevels[0]);
 	}
+
+	AdvancedPreviewSettingsWidget = SNew(SAdvancedPreviewDetailsTab)
+		.PreviewScenePtr(&Viewport->GetPreviewScene());
 }
 
 FName FStaticMeshEditor::GetToolkitFName() const
diff --git a/Engine/Source/Editor/StaticMeshEditor/Private/StaticMeshEditor.h b/Engine/Source/Editor/StaticMeshEditor/Private/StaticMeshEditor.h
--- a/Engine/Source/Editor/StaticMeshEditor/Private/StaticMeshEditor.h
+++ b/Engine/Source/Editor/StaticMeshEditor/Private/StaticMeshEditor.h
@@ -115,6 +115,7 @@ private:
 	TSharedRef<SDockTab> SpawnTab_Properties(const FSpawnTabArgs& Args);
 	TSharedRef<SDockTab> SpawnTab_SocketManager(const FSpawnTabArgs& Args);
 	TSharedRef<SDockTab> SpawnTab_Collision(const FSpawnTabArgs& Args);
+	TSharedRef<SDockTab> SpawnTab_PreviewSceneSettings(const FSpawnTabArgs& Args);
 
 private:
 	/** Binds commands associated with the Static Mesh Editor. */
@@ -304,6 +305,9 @@ private:
 	/** Information on the selected collision primitives */
 	TArray<FPrimData> SelectedPrims;
 
+	/** Scene preview settings widget */
+	TSharedPtr< class SAdvancedPreviewDetailsTab> AdvancedPreviewSettingsWidget;
+
 	/** Misc consts */
 	const float MinPrimSize;
 	const FVector OverlapNudge;
@@ -313,4 +317,5 @@ private:
 	static const FName PropertiesTabId;
 	static const FName SocketManagerTabId;
 	static const FName CollisionTabId;
+	static const FName PreviewSceneSettingsTabId;
 };
diff --git a/Engine/Source/Editor/StaticMeshEditor/Private/StaticMeshEditorViewportClient.cpp b/Engine/Source/Editor/StaticMeshEditor/Private/StaticMeshEditorViewportClient.cpp
--- a/Engine/Source/Editor/StaticMeshEditor/Private/StaticMeshEditorViewportClient.cpp
+++ b/Engine/Source/Editor/StaticMeshEditor/Private/StaticMeshEditorViewportClient.cpp
@@ -47,7 +47,7 @@ namespace {
 	static float AmbientCubemapIntensity = 0.4f;
 }
 
-FStaticMeshEditorViewportClient::FStaticMeshEditorViewportClient(TWeakPtr<IStaticMeshEditor> InStaticMeshEditor, const TSharedRef<SStaticMeshEditorViewport>& InStaticMeshEditorViewport, FPreviewScene& InPreviewScene, UStaticMesh* InPreviewStaticMesh, UStaticMeshComponent* InPreviewStaticMeshComponent)
+FStaticMeshEditorViewportClient::FStaticMeshEditorViewportClient(TWeakPtr<IStaticMeshEditor> InStaticMeshEditor, const TSharedRef<SStaticMeshEditorViewport>& InStaticMeshEditorViewport, FAdvancedPreviewScene& InPreviewScene, UStaticMesh* InPreviewStaticMesh, UStaticMeshComponent* InPreviewStaticMeshComponent)
 	: FEditorViewportClient(nullptr, &InPreviewScene, StaticCastSharedRef<SEditorViewport>(InStaticMeshEditorViewport))
 	, StaticMeshEditorPtr(InStaticMeshEditor)
 	, StaticMeshEditorViewportPtr(InStaticMeshEditorViewport)
@@ -69,7 +69,6 @@ FStaticMeshEditorViewportClient::FStaticMeshEditorViewportClient(TWeakPtr<IStati
 
 	WidgetMode = FWidget::WM_None;
 
-	EngineShowFlags.DisableAdvancedFeatures();
 	EngineShowFlags.SetSeparateTranslucency(true);
 	EngineShowFlags.SetSnap(0);
 	EngineShowFlags.SetCompositeEditorPrimitives(true);
@@ -840,7 +839,23 @@ bool FStaticMeshEditorViewportClient::InputKey(FViewport* InViewport, int32 Cont
 
 bool FStaticMeshEditorViewportClient::InputAxis(FViewport* InViewport, int32 ControllerId, FKey Key, float Delta, float DeltaTime, int32 NumSamples, bool bGamepad)
 {
-	return FEditorViewportClient::InputAxis(InViewport,ControllerId,Key,Delta,DeltaTime,NumSamples,bGamepad);
+	bool bResult = true;
+	
+	if (!bDisableInput)
+	{
+		FAdvancedPreviewScene* AdvancedScene = (FAdvancedPreviewScene*)PreviewScene;
+		bResult = AdvancedScene->HandleViewportInput(InViewport, ControllerId, Key, Delta, DeltaTime, NumSamples, bGamepad);
+		if (bResult)
+		{
+			Invalidate();
+		}
+		else
+		{
+			bResult = FEditorViewportClient::InputAxis(InViewport, ControllerId, Key, Delta, DeltaTime, NumSamples, bGamepad);
+		}
+	}
+
+	return bResult;
 }
 
 void FStaticMeshEditorViewportClient::ProcessClick(class FSceneView& InView, class HHitProxy* HitProxy, FKey Key, EInputEvent Event, uint32 HitX, uint32 HitY)
@@ -1144,15 +1159,6 @@ void FStaticMeshEditorViewportClient::ProcessClick(class FSceneView& InView, cla
 	Invalidate();
 }
 
-FSceneView* FStaticMeshEditorViewportClient::CalcSceneView(FSceneViewFamily* ViewFamily, const EStereoscopicPass StereoPass)
-{
-	FSceneView* SceneView = FEditorViewportClient::CalcSceneView(ViewFamily, StereoPass);
-	FFinalPostProcessSettings::FCubemapEntry& CubemapEntry = *new(SceneView->FinalPostProcessSettings.ContributingCubemaps) FFinalPostProcessSettings::FCubemapEntry;
-	CubemapEntry.AmbientCubemap = GUnrealEd->GetThumbnailManager()->AmbientCubemap;
-	CubemapEntry.AmbientCubemapTintMulScaleValue = FLinearColor::White * AmbientCubemapIntensity;
-	return SceneView;
-}
-
 void FStaticMeshEditorViewportClient::PerspectiveCameraMoved()
 {
 	FEditorViewportClient::PerspectiveCameraMoved();
diff --git a/Engine/Source/Editor/StaticMeshEditor/Private/StaticMeshEditorViewportClient.h b/Engine/Source/Editor/StaticMeshEditor/Private/StaticMeshEditorViewportClient.h
--- a/Engine/Source/Editor/StaticMeshEditor/Private/StaticMeshEditorViewportClient.h
+++ b/Engine/Source/Editor/StaticMeshEditor/Private/StaticMeshEditorViewportClient.h
@@ -5,12 +5,13 @@
 #include "EditorViewportClient.h"
 
 class SStaticMeshEditorViewport;
+class FAdvancedPreviewScene;
 
 /** Viewport Client for the preview viewport */
 class FStaticMeshEditorViewportClient : public FEditorViewportClient, public TSharedFromThis<FStaticMeshEditorViewportClient>
 {
 public:
-	FStaticMeshEditorViewportClient(TWeakPtr<IStaticMeshEditor> InStaticMeshEditor, const TSharedRef<SStaticMeshEditorViewport>& InStaticMeshEditorViewport, FPreviewScene& InPreviewScene, UStaticMesh* InPreviewStaticMesh, UStaticMeshComponent* InPreviewStaticMeshComponent);
+	FStaticMeshEditorViewportClient(TWeakPtr<IStaticMeshEditor> InStaticMeshEditor, const TSharedRef<SStaticMeshEditorViewport>& InStaticMeshEditorViewport, FAdvancedPreviewScene& InPreviewScene, UStaticMesh* InPreviewStaticMesh, UStaticMeshComponent* InPreviewStaticMeshComponent);
 
 	// FEditorViewportClient interface
 	virtual void MouseMove(FViewport* Viewport,int32 x, int32 y) override;
@@ -31,7 +32,6 @@ public:
 	virtual FMatrix GetWidgetCoordSystem() const override;
 	virtual ECoordSystem GetWidgetCoordSystemSpace() const override { return COORD_Local; }
 	virtual bool ShouldOrbitCamera() const override;
-	virtual FSceneView* CalcSceneView(FSceneViewFamily* ViewFamily, const EStereoscopicPass StereoPass = eSSP_FULL) override;
 
 	/** 
 	 *	Updates the static mesh and static mesh component being used in the Static Mesh Editor.
diff --git a/Engine/Source/Editor/UnrealEd/Classes/Editor/EditorPerProjectUserSettings.h b/Engine/Source/Editor/UnrealEd/Classes/Editor/EditorPerProjectUserSettings.h
--- a/Engine/Source/Editor/UnrealEd/Classes/Editor/EditorPerProjectUserSettings.h
+++ b/Engine/Source/Editor/UnrealEd/Classes/Editor/EditorPerProjectUserSettings.h
@@ -130,6 +130,9 @@ class UEditorPerProjectUserSettings : public UObject
 	class UBlueprintPaletteFavorites* BlueprintFavorites;
 
 public:
+	// Per project user settings for which asset viewer profile should be used
+	UPROPERTY(config)
+	int32 AssetViewerProfileIndex;
 
 	UPROPERTY(config)
 	int32 MaterialQualityLevel;
diff --git a/Engine/Source/Editor/UnrealEd/Private/Factories/EditorFactories.cpp b/Engine/Source/Editor/UnrealEd/Private/Factories/EditorFactories.cpp
--- a/Engine/Source/Editor/UnrealEd/Private/Factories/EditorFactories.cpp
+++ b/Engine/Source/Editor/UnrealEd/Private/Factories/EditorFactories.cpp
@@ -5449,11 +5449,13 @@ void FCustomizableTextObjectFactory::ProcessBuffer(UObject* InParent, EObjectFla
 					ObjArchetype = LoadObject<UObject>(nullptr, *ObjArchetypeName, nullptr, LOAD_None, nullptr);
 				}
 
+				UObject* ObjectParent = InParent ? InParent : GetParentForNewObject(ObjClass);
+
 				// Make sure this name is not used by anything else. Will rename other stuff if necessary
-				ClearObjectNameUsage(InParent, ObjName);
+				ClearObjectNameUsage(ObjectParent, ObjName);
 
 				// Spawn the object and reset it's archetype
-				UObject* CreatedObject = NewObject<UObject>(InParent, ObjClass, ObjName, Flags, ObjArchetype, !!InParent, &InstanceGraph);
+				UObject* CreatedObject = NewObject<UObject>(ObjectParent, ObjClass, ObjName, Flags, ObjArchetype, !!ObjectParent, &InstanceGraph);
 
 				// Get property text for the new object.
 				FString PropText, PropLine;
@@ -5498,6 +5500,10 @@ void FCustomizableTextObjectFactory::ProcessBuffer(UObject* InParent, EObjectFla
 		{
 			--NestedDepth;
 		}
+		else
+		{
+			ProcessUnidentifiedLine(StrLine);
+		}
 	}
 
 	// Apply the property text to each of the created objects
diff --git a/Engine/Source/Editor/UnrealEd/Private/PhysicsAssetUtils.cpp b/Engine/Source/Editor/UnrealEd/Private/PhysicsAssetUtils.cpp
--- a/Engine/Source/Editor/UnrealEd/Private/PhysicsAssetUtils.cpp
+++ b/Engine/Source/Editor/UnrealEd/Private/PhysicsAssetUtils.cpp
@@ -183,9 +183,9 @@ bool CreateFromSkeletalMeshInternal(UPhysicsAsset* PhysicsAsset, USkeletalMesh*
 				FMatrix RelTM = SkelMesh->GetRefPoseMatrix(i);
 
 				// set angular constraint mode
-				CS->DefaultInstance.AngularSwing1Motion = Params.AngularConstraintMode;
-				CS->DefaultInstance.AngularSwing2Motion = Params.AngularConstraintMode;
-				CS->DefaultInstance.AngularTwistMotion = Params.AngularConstraintMode;
+				CS->DefaultInstance.SetAngularSwing1Motion(Params.AngularConstraintMode);
+				CS->DefaultInstance.SetAngularSwing2Motion(Params.AngularConstraintMode);
+				CS->DefaultInstance.SetAngularTwistMotion(Params.AngularConstraintMode);
 
 				// Place joint at origin of child
 				CS->DefaultInstance.ConstraintBone1 = BoneName;
diff --git a/Engine/Source/Editor/UnrealEd/Private/Settings/EditorPerProjectUserSettings.cpp b/Engine/Source/Editor/UnrealEd/Private/Settings/EditorPerProjectUserSettings.cpp
--- a/Engine/Source/Editor/UnrealEd/Private/Settings/EditorPerProjectUserSettings.cpp
+++ b/Engine/Source/Editor/UnrealEd/Private/Settings/EditorPerProjectUserSettings.cpp
@@ -14,6 +14,7 @@ UEditorPerProjectUserSettings::UEditorPerProjectUserSettings(const FObjectInitia
 	bMonitorEditorPerformance = true;
 	BlueprintFavorites = CreateDefaultSubobject<UBlueprintPaletteFavorites>(TEXT("BlueprintFavorites"));
 	SCSViewportCameraSpeed = 4;
+	AssetViewerProfileIndex = 0;
 
 	SimplygonServerIP = "127.0.0.1";
 }
diff --git a/Engine/Source/Editor/UnrealEd/Public/Factories.h b/Engine/Source/Editor/UnrealEd/Public/Factories.h
--- a/Engine/Source/Editor/UnrealEd/Public/Factories.h
+++ b/Engine/Source/Editor/UnrealEd/Public/Factories.h
@@ -49,6 +49,12 @@ protected:
 
 	/** Util to ensure that InName is a valid name for a new object within InParent. Will rename any existing object within InParent if it is called InName. */
 	static void ClearObjectNameUsage(UObject* InParent, FName InName);
+
+	/** If we cant do anything with the line ourselves hand off to child class */
+	virtual void ProcessUnidentifiedLine(const FString& StrLine) {}
+
+	/** Allow child class to override new object parent (only called when parent supplied to ProcessBuffer is NULL */
+	virtual UObject* GetParentForNewObject(const UClass* ObjClass) {return NULL;}
 };
 
 
diff --git a/Engine/Source/Editor/UnrealEd/Public/SAdvancedPreviewDetailsTab.cpp b/Engine/Source/Editor/UnrealEd/Public/SAdvancedPreviewDetailsTab.cpp
new file mode 100644
--- /dev/null
+++ b/Engine/Source/Editor/UnrealEd/Public/SAdvancedPreviewDetailsTab.cpp
@@ -0,0 +1,217 @@
+// Copyright 1998-2016 Epic Games, Inc. All Rights Reserved.
+
+#include "UnrealEd.h"
+#include "SlateBasics.h"
+#include "SAdvancedPreviewDetailsTab.h"
+#include "STextComboBox.h"
+#include "IDetailsView.h"
+#include "ModuleManager.h"
+#include "PropertyEditorModule.h"
+
+#include "AssetViewerSettings.h"
+#include "ScopedTransaction.h"
+
+#include "AdvancedPreviewScene.h"
+
+#define LOCTEXT_NAMESPACE "SPrettyPreview"
+
+SAdvancedPreviewDetailsTab::SAdvancedPreviewDetailsTab()
+{
+	DefaultSettings = UAssetViewerSettings::Get();
+	if (DefaultSettings)
+	{
+		RefreshDelegate = DefaultSettings->OnAssetViewerSettingsChanged().AddRaw(this, &SAdvancedPreviewDetailsTab::OnAssetViewerSettingsRefresh);
+			
+		AddRemoveProfileDelegate = DefaultSettings->OnAssetViewerProfileAddRemoved().AddLambda([this]() { this->Refresh(); });
+	}
+
+	PerProjectSettings = GetMutableDefault<UEditorPerProjectUserSettings>();
+
+}
+
+SAdvancedPreviewDetailsTab::~SAdvancedPreviewDetailsTab()
+{
+	DefaultSettings = UAssetViewerSettings::Get();
+	if (DefaultSettings)
+	{
+		DefaultSettings->OnAssetViewerSettingsChanged().Remove(RefreshDelegate);
+		DefaultSettings->OnAssetViewerProfileAddRemoved().Remove(AddRemoveProfileDelegate);
+	}
+}
+
+void SAdvancedPreviewDetailsTab::Construct(const FArguments& InArgs)
+{
+	PreviewScene = InArgs._PreviewScenePtr;	
+	DefaultSettings = UAssetViewerSettings::Get();
+	ProfileIndex = PerProjectSettings->AssetViewerProfileIndex;
+		
+	CreateSettingsView();
+
+	this->ChildSlot
+	[
+		SNew(SVerticalBox)
+		+ SVerticalBox::Slot()
+		.Padding(2.0f, 1.0f, 2.0f, 1.0f )
+		.AutoHeight()
+		[
+			SNew(SHorizontalBox)
+			+ SHorizontalBox::Slot()
+			[
+				SNew(STextBlock)
+				.Text(LOCTEXT("SceneProfileSettingsLabel", "Scene Profile Settings"))				
+			]
+		]
+		+SVerticalBox::Slot()		
+		.Padding(2.0f, 1.0f, 2.0f, 1.0f)
+		.AutoHeight()
+		[
+			SNew(SHorizontalBox)
+			+SHorizontalBox::Slot()
+			.Padding(2.0f)
+			[
+				SAssignNew(ProfileComboBox, STextComboBox)
+				.OptionsSource(&ProfileNames)
+				.OnSelectionChanged(this, &SAdvancedPreviewDetailsTab::ComboBoxSelectionChanged)				
+				.ToolTipText(LOCTEXT("SceneProfileComboBoxToolTip", "Allows for switching between scene profiles."))
+				.IsEnabled_Lambda([this]() -> bool
+				{
+					return ProfileNames.Num() > 1;
+				})
+			]
+			+ SHorizontalBox::Slot()
+			.Padding(2.0f)
+			.AutoWidth()
+			[
+				SNew(SButton)
+				.OnClicked(this, &SAdvancedPreviewDetailsTab::AddProfileButtonClick)
+				.Text(LOCTEXT("AddProfileButton", "Add Profile"))
+				.ToolTipText(LOCTEXT("SceneProfileAddProfile", "Adds a new profile."))
+			]
+			+ SHorizontalBox::Slot()
+			.Padding(2.0f)
+			.AutoWidth()
+			[
+				SNew(SButton)
+				.OnClicked(this, &SAdvancedPreviewDetailsTab::RemoveProfileButtonClick)
+				.Text(LOCTEXT("RemoveProfileButton", "Remove Profile"))
+				.ToolTipText(LOCTEXT("SceneProfileRemoveProfile", "Removes the currently selected profile."))
+				.IsEnabled_Lambda([this]()->bool
+				{
+					return ProfileNames.Num() > 1; 
+				})
+			]
+		]
+
+		+ SVerticalBox::Slot()
+		.Padding(2.0f, 1.0f, 2.0f, 1.0f)
+		[
+			SNew(SHorizontalBox)
+			+SHorizontalBox::Slot()
+			[
+				SettingsView->AsShared()
+			]
+		]
+	];
+
+	UpdateProfileNames();
+	UpdateSettingsView();
+}
+
+void SAdvancedPreviewDetailsTab::ComboBoxSelectionChanged(TSharedPtr<FString> NewSelection, ESelectInfo::Type /*SelectInfo*/)
+{
+	int32 NewSelectionIndex;
+
+	if (ProfileNames.Find(NewSelection, NewSelectionIndex))
+	{
+		ProfileIndex = NewSelectionIndex;
+		PerProjectSettings->AssetViewerProfileIndex = ProfileIndex;
+		UpdateSettingsView();	
+		PreviewScene->SetProfileIndex(ProfileIndex);
+	}
+}
+
+void SAdvancedPreviewDetailsTab::UpdateSettingsView()
+{	
+	SettingsView->SetObject(DefaultSettings, true);
+}
+
+void SAdvancedPreviewDetailsTab::UpdateProfileNames()
+{
+	checkf(DefaultSettings->Profiles.Num(), TEXT("There should always be at least one profile available"));
+	ProfileNames.Empty();
+	for (FPreviewSceneProfile& Profile : DefaultSettings->Profiles)
+	{
+		ProfileNames.Add(TSharedPtr<FString>(new FString(Profile.ProfileName)));
+	}
+
+	ProfileComboBox->RefreshOptions();
+	ProfileComboBox->SetSelectedItem(ProfileNames[ProfileIndex]);
+}
+
+FReply SAdvancedPreviewDetailsTab::AddProfileButtonClick()
+{
+	// Add new profile to settings instance
+	DefaultSettings->Profiles.AddDefaulted();
+	FPreviewSceneProfile& NewProfile = DefaultSettings->Profiles.Last();
+	NewProfile.ProfileName = FString::Printf(TEXT("Profile_%i"), DefaultSettings->Profiles.Num());
+	DefaultSettings->PostEditChange();
+	
+	return FReply::Handled();
+}
+
+FReply SAdvancedPreviewDetailsTab::RemoveProfileButtonClick()
+{
+	// Remove currently selected profile 
+	DefaultSettings->Profiles.RemoveAt(ProfileIndex);
+	ProfileIndex = DefaultSettings->Profiles.IsValidIndex(ProfileIndex - 1 ) ? ProfileIndex - 1 : 0;
+	PerProjectSettings->AssetViewerProfileIndex = ProfileIndex;
+	DefaultSettings->PostEditChange();
+	
+	return FReply::Handled();
+}
+
+void SAdvancedPreviewDetailsTab::OnAssetViewerSettingsRefresh(const FName& InPropertyName)
+{
+	const bool bUpdateEnvironment = (InPropertyName == GET_MEMBER_NAME_CHECKED(FPreviewSceneProfile, EnvironmentCubeMap)) || (InPropertyName == GET_MEMBER_NAME_CHECKED(FPreviewSceneProfile, LightingRigRotation) || (InPropertyName == GET_MEMBER_NAME_CHECKED(UAssetViewerSettings, Profiles)));
+	const bool bUpdateSkyLight = bUpdateEnvironment || (InPropertyName == GET_MEMBER_NAME_CHECKED(FPreviewSceneProfile, SkyLightIntensity) || (InPropertyName == GET_MEMBER_NAME_CHECKED(UAssetViewerSettings, Profiles)));
+	const bool bUpdateDirectionalLight = (InPropertyName == GET_MEMBER_NAME_CHECKED(FPreviewSceneProfile, DirectionalLightIntensity)) || (InPropertyName == GET_MEMBER_NAME_CHECKED(FPreviewSceneProfile, DirectionalLightColor));
+	const bool bUpdatePostProcessing = ( InPropertyName == GET_MEMBER_NAME_CHECKED(FPreviewSceneProfile, PostProcessingSettings));
+
+	if (InPropertyName == GET_MEMBER_NAME_CHECKED(FPreviewSceneProfile, ProfileName))
+	{
+		UpdateProfileNames();
+	}
+
+	PreviewScene->UpdateScene(DefaultSettings->Profiles[ProfileIndex], bUpdateSkyLight, bUpdateEnvironment, bUpdatePostProcessing, bUpdateDirectionalLight);
+}
+
+void SAdvancedPreviewDetailsTab::CreateSettingsView()
+{	
+	// Create a property view
+	FPropertyEditorModule& EditModule = FModuleManager::Get().GetModuleChecked<FPropertyEditorModule>("PropertyEditor");
+
+	FDetailsViewArgs DetailsViewArgs(
+		/*bUpdateFromSelection=*/ false,
+		/*bLockable=*/ false,
+		/*bAllowSearch=*/ false,
+		FDetailsViewArgs::HideNameArea,
+		/*bHideSelectionTip=*/ true,
+		/*InNotifyHook=*/ nullptr,
+		/*InSearchInitialKeyFocus=*/ false,
+		/*InViewIdentifier=*/ NAME_None);
+	DetailsViewArgs.DefaultsOnlyVisibility = FDetailsViewArgs::EEditDefaultsOnlyNodeVisibility::Automatic;
+	DetailsViewArgs.bShowOptions = false;
+
+	SettingsView = EditModule.CreateDetailView(DetailsViewArgs);
+	UpdateSettingsView();
+}
+
+void SAdvancedPreviewDetailsTab::Refresh()
+{	
+	ProfileIndex = PerProjectSettings->AssetViewerProfileIndex;	
+	UpdateProfileNames();
+	PreviewScene->SetProfileIndex(ProfileIndex);
+	UpdateSettingsView();
+}
+
+#undef LOCTEXT_NAMESPACE
\ No newline at end of file
diff --git a/Engine/Source/Editor/UnrealEd/Public/SAdvancedPreviewDetailsTab.h b/Engine/Source/Editor/UnrealEd/Public/SAdvancedPreviewDetailsTab.h
new file mode 100644
--- /dev/null
+++ b/Engine/Source/Editor/UnrealEd/Public/SAdvancedPreviewDetailsTab.h
@@ -0,0 +1,56 @@
+// Copyright 1998-2016 Epic Games, Inc. All Rights Reserved.
+#pragma once
+
+#include "NotifyHook.h"
+
+class IDetailsView;
+class FAdvancedPreviewScene;
+class UAssetViewerSettings;
+class UEditorPerProjectUserSettings;
+
+class UNREALED_API SAdvancedPreviewDetailsTab : public SCompoundWidget
+{
+public:
+	SLATE_BEGIN_ARGS(SAdvancedPreviewDetailsTab) : 
+		_PreviewScenePtr()
+	{
+	}
+		/** The Static Mesh Editor this tool is associated with. */
+		SLATE_ARGUMENT(FAdvancedPreviewScene*, PreviewScenePtr)
+
+	SLATE_END_ARGS()
+
+public:
+	/** **/
+	SAdvancedPreviewDetailsTab();
+	~SAdvancedPreviewDetailsTab();
+
+	/** SWidget functions */
+	void Construct(const FArguments& InArgs);
+
+	void Refresh();
+
+protected:
+	void CreateSettingsView();
+	void ComboBoxSelectionChanged(TSharedPtr<FString> NewSelection, ESelectInfo::Type /*SelectInfo*/);
+	void UpdateSettingsView();
+	void UpdateProfileNames();
+	FReply AddProfileButtonClick();
+	FReply RemoveProfileButtonClick();
+protected:
+	void OnAssetViewerSettingsRefresh(const FName& InPropertyName);
+protected:
+	/** Property viewing widget */
+	TSharedPtr<IDetailsView> SettingsView;
+	TSharedPtr<STextComboBox> ProfileComboBox;
+	FAdvancedPreviewScene* PreviewScene;
+	UAssetViewerSettings* DefaultSettings;
+
+	TArray<TSharedPtr<FString>> ProfileNames;
+	int32 ProfileIndex;
+
+	FDelegateHandle RefreshDelegate;
+	FDelegateHandle AddRemoveProfileDelegate;
+
+	UEditorPerProjectUserSettings* PerProjectSettings;
+};
diff --git a/Engine/Source/Programs/AutomationTool/Scripts/StreamCopyDescription.Automation.cs b/Engine/Source/Programs/AutomationTool/Scripts/StreamCopyDescription.Automation.cs
--- a/Engine/Source/Programs/AutomationTool/Scripts/StreamCopyDescription.Automation.cs
+++ b/Engine/Source/Programs/AutomationTool/Scripts/StreamCopyDescription.Automation.cs
@@ -82,7 +82,7 @@ namespace AutomationTool
 				foreach (string Line in Descriptions.Split('\n'))
 				{
 					string TrimLine = Line.TrimStart();
-					if(!TrimLine.StartsWith("#codereview") && !TrimLine.StartsWith("#rb"))
+					if(!TrimLine.StartsWith("#codereview") && !TrimLine.StartsWith("#rb") && !TrimLine.StartsWith("#lockdown"))
 					{
 						Writer.WriteLine(Line);
 					}
diff --git a/Engine/Source/Runtime/AnimGraphRuntime/Private/BoneControllers/AnimNode_AnimDynamics.cpp b/Engine/Source/Runtime/AnimGraphRuntime/Private/BoneControllers/AnimNode_AnimDynamics.cpp
--- a/Engine/Source/Runtime/AnimGraphRuntime/Private/BoneControllers/AnimNode_AnimDynamics.cpp
+++ b/Engine/Source/Runtime/AnimGraphRuntime/Private/BoneControllers/AnimNode_AnimDynamics.cpp
@@ -94,6 +94,13 @@ void FAnimNode_AnimDynamics::EvaluateBoneTransforms(USkeletalMeshComponent* Skel
 
 	if (CVarEnableDynamics.GetValueOnAnyThread() == 1 && bEnabledForLod)
 	{
+		if(LastSimSpace != SimulationSpace)
+		{
+			// Our sim space has been changed since our last update, we need to convert all of our
+			// body transforms into the new space.
+			ConvertSimulationSpace(SkelComp, MeshBases, LastSimSpace, SimulationSpace);
+		}
+
 		// Pretty nasty - but there isn't really a good way to get clean bone transforms (without the modification from
 		// previous runs) so we have to initialize here, checking often so we can restart a simulation in the editor.
 		if (bRequiresInit)
@@ -211,11 +218,14 @@ void FAnimNode_AnimDynamics::EvaluateBoneTransforms(USkeletalMeshComponent* Skel
 
 				FTransform NewBoneTransform(CurrentBody.Pose.Orientation, CurrentBody.Pose.Position + CurrentBody.Pose.Orientation.RotateVector(JointOffsets[Idx]));
 
-				NewBoneTransform = GetComponentSpaceTransformFromSimSpace(SkelComp, MeshBases, BoneIndex, NewBoneTransform);
+				NewBoneTransform = GetComponentSpaceTransformFromSimSpace(SimulationSpace, SkelComp, MeshBases, NewBoneTransform);
 
 				OutBoneTransforms.Add(FBoneTransform(BoneIndex, NewBoneTransform));
 			}
 		}
+
+		// Store our sim space incase it changes
+		LastSimSpace = SimulationSpace;
 	}
 }
 
@@ -644,9 +654,16 @@ FTransform FAnimNode_AnimDynamics::GetBoneTransformInSimSpace(USkeletalMeshCompo
 {
 	FTransform Transform = MeshBases.GetComponentSpaceTransform(BoneIndex);
 
-	switch(SimulationSpace)
+	return GetSimSpaceTransformFromComponentSpace(SimulationSpace, SkelComp, MeshBases, Transform);
+}
+
+FTransform FAnimNode_AnimDynamics::GetComponentSpaceTransformFromSimSpace(AnimPhysSimSpaceType SimSpace, USkeletalMeshComponent* SkelComp, FCSPose<FCompactPose>& MeshBases, const FTransform& InSimTransform)
+{
+	FTransform OutTransform = InSimTransform;
+
+	switch(SimSpace)
 	{
-	// Change nothing, already in component space
+		// Change nothing, already in component space
 	case AnimPhysSimSpaceType::Component:
 	{
 		break;
@@ -659,12 +676,11 @@ FTransform FAnimNode_AnimDynamics::GetBoneTransformInSimSpace(USkeletalMeshCompo
 		{
 			AActor* OwningActor = Cast<AActor>(OuterObject);
 			check(OwningActor);
-
-			FTransform WorldTransform = Transform * SkelComp->ComponentToWorld;
-			WorldTransform.SetToRelativeTransform(OwningActor->GetActorTransform());
-			Transform = WorldTransform;
-		}
 		
+			FTransform ComponentTransform(SkelComp->RelativeRotation, SkelComp->RelativeLocation, SkelComp->RelativeScale3D);
+			OutTransform = OutTransform * ComponentTransform.Inverse();
+		}
+
 		break;
 	}
 
@@ -675,7 +691,7 @@ FTransform FAnimNode_AnimDynamics::GetBoneTransformInSimSpace(USkeletalMeshCompo
 		FCompactPoseBoneIndex RootBoneCompactIndex(0);
 
 		FTransform RelativeBoneTransform = MeshBases.GetComponentSpaceTransform(RootBoneCompactIndex);
-		Transform = Transform.GetRelativeTransform(RelativeBoneTransform);
+		OutTransform = OutTransform * RelativeBoneTransform;
 
 		break;
 	}
@@ -686,30 +702,29 @@ FTransform FAnimNode_AnimDynamics::GetBoneTransformInSimSpace(USkeletalMeshCompo
 		if(RelativeSpaceBone.IsValid(RequiredBones))
 		{
 			FTransform RelativeBoneTransform = MeshBases.GetComponentSpaceTransform(RelativeSpaceBone.GetCompactPoseIndex(RequiredBones));
-			Transform = Transform.GetRelativeTransform(RelativeBoneTransform);
+			OutTransform = OutTransform * RelativeBoneTransform;
 		}
 
 		break;
 	}
-
 	case AnimPhysSimSpaceType::World:
 	{
-		// Out to world space
-		FAnimationRuntime::ConvertCSTransformToBoneSpace(SkelComp, MeshBases, Transform, BoneIndex, BCS_WorldSpace);
+		OutTransform *= SkelComp->ComponentToWorld.Inverse();
 	}
 
 	default:
 		break;
 	}
 
-	return Transform;
+	return OutTransform;
 }
 
-FTransform FAnimNode_AnimDynamics::GetComponentSpaceTransformFromSimSpace(USkeletalMeshComponent* SkelComp, FCSPose<FCompactPose>& MeshBases, const FCompactPoseBoneIndex& BoneIndex, const FTransform& InSimTransform)
+
+FTransform FAnimNode_AnimDynamics::GetSimSpaceTransformFromComponentSpace(AnimPhysSimSpaceType SimSpace, USkeletalMeshComponent* SkelComp, FCSPose<FCompactPose>& MeshBases, const FTransform& InComponentTransform)
 {
-	FTransform OutTransform = InSimTransform;
+	FTransform ResultTransform = InComponentTransform;
 
-	switch(SimulationSpace)
+	switch(SimSpace)
 	{
 		// Change nothing, already in component space
 	case AnimPhysSimSpaceType::Component:
@@ -724,9 +739,10 @@ FTransform FAnimNode_AnimDynamics::GetComponentSpaceTransformFromSimSpace(USkele
 		{
 			AActor* OwningActor = Cast<AActor>(OuterObject);
 			check(OwningActor);
-		
-			FTransform ComponentTransform(SkelComp->RelativeRotation, SkelComp->RelativeLocation, SkelComp->RelativeScale3D);
-			OutTransform = OutTransform * ComponentTransform.Inverse();
+
+			FTransform WorldTransform = ResultTransform * SkelComp->ComponentToWorld;
+			WorldTransform.SetToRelativeTransform(OwningActor->GetActorTransform());
+			ResultTransform = WorldTransform;
 		}
 
 		break;
@@ -739,7 +755,7 @@ FTransform FAnimNode_AnimDynamics::GetComponentSpaceTransformFromSimSpace(USkele
 		FCompactPoseBoneIndex RootBoneCompactIndex(0);
 
 		FTransform RelativeBoneTransform = MeshBases.GetComponentSpaceTransform(RootBoneCompactIndex);
-		OutTransform = OutTransform * RelativeBoneTransform;
+		ResultTransform = ResultTransform.GetRelativeTransform(RelativeBoneTransform);
 
 		break;
 	}
@@ -750,24 +766,25 @@ FTransform FAnimNode_AnimDynamics::GetComponentSpaceTransformFromSimSpace(USkele
 		if(RelativeSpaceBone.IsValid(RequiredBones))
 		{
 			FTransform RelativeBoneTransform = MeshBases.GetComponentSpaceTransform(RelativeSpaceBone.GetCompactPoseIndex(RequiredBones));
-			OutTransform = OutTransform * RelativeBoneTransform;
+			ResultTransform = ResultTransform.GetRelativeTransform(RelativeBoneTransform);
 		}
 
 		break;
 	}
+
 	case AnimPhysSimSpaceType::World:
 	{
-		OutTransform *= SkelComp->ComponentToWorld.Inverse();
+		// Out to world space
+		ResultTransform *= SkelComp->ComponentToWorld;
 	}
 
 	default:
 		break;
 	}
 
-	return OutTransform;
+	return ResultTransform;
 }
 
-
 FVector FAnimNode_AnimDynamics::TransformWorldVectorToSimSpace(USkeletalMeshComponent* SkelComp, FCSPose<FCompactPose>& MeshBases, const FVector& InVec)
 {
 	FVector OutVec = InVec;
@@ -832,3 +849,25 @@ FVector FAnimNode_AnimDynamics::TransformWorldVectorToSimSpace(USkeletalMeshComp
 	return OutVec;
 }
 
+void FAnimNode_AnimDynamics::ConvertSimulationSpace(USkeletalMeshComponent* SkelComp, FCSPose<FCompactPose>& MeshBases, AnimPhysSimSpaceType From, AnimPhysSimSpaceType To)
+{
+	for(FAnimPhysRigidBody* Body : BaseBodyPtrs)
+	{
+		if(!Body)
+		{
+			return;
+		}
+
+		// Get transform
+		FTransform BodyTransform(Body->Pose.Orientation, Body->Pose.Position);
+		// Out to component space
+		BodyTransform = GetComponentSpaceTransformFromSimSpace(LastSimSpace, SkelComp, MeshBases, BodyTransform);
+		// In to new space
+		BodyTransform = GetSimSpaceTransformFromComponentSpace(SimulationSpace, SkelComp, MeshBases, BodyTransform);
+
+		// Push back to body
+		Body->Pose.Orientation = BodyTransform.GetRotation();
+		Body->Pose.Position = BodyTransform.GetTranslation();
+	}
+}
+
diff --git a/Engine/Source/Runtime/AnimGraphRuntime/Private/BoneControllers/AnimNode_OrientationDriver.cpp b/Engine/Source/Runtime/AnimGraphRuntime/Private/BoneControllers/AnimNode_OrientationDriver.cpp
new file mode 100644
--- /dev/null
+++ b/Engine/Source/Runtime/AnimGraphRuntime/Private/BoneControllers/AnimNode_OrientationDriver.cpp
@@ -0,0 +1,181 @@
+// Copyright 1998-2016 Epic Games, Inc. All Rights Reserved.
+
+#include "AnimGraphRuntimePrivatePCH.h"
+#include "BoneControllers/AnimNode_OrientationDriver.h"
+
+#include "AnimInstanceProxy.h"
+
+//////////////////////////////////////////////////////////////////////////
+
+void FOrientationDriverParamSet::AddParam(const FOrientationDriverParam& InParam, float InScale)
+{
+	bool bFoundParam = false;
+
+	// Look to see if this param already exists in this set
+	for (FOrientationDriverParam& Param : Params)
+	{
+		if (Param.ParameterName == InParam.ParameterName)
+		{
+			Param.ParameterValue += (InParam.ParameterValue * InScale);
+			bFoundParam = true;
+		}
+	}
+
+	// If not found, add new entry and set to supplied value
+	if (!bFoundParam)
+	{
+		int32 NewParamIndex = Params.AddZeroed();
+		FOrientationDriverParam& Param = Params[NewParamIndex];
+		Param.ParameterName = InParam.ParameterName;
+		Param.ParameterValue = (InParam.ParameterValue * InScale);
+	}
+}
+
+void FOrientationDriverParamSet::AddParams(const TArray<FOrientationDriverParam>& InParams, float InScale)
+{
+	for (const FOrientationDriverParam& InParam : InParams)
+	{
+		AddParam(InParam, InScale);
+	}
+}
+
+void FOrientationDriverParamSet::ScaleAllParams(float InScale)
+{
+	for (FOrientationDriverParam& Param : Params)
+	{
+		Param.ParameterValue *= InScale;
+	}
+}
+
+void FOrientationDriverParamSet::ClearParams()
+{
+	Params.Empty();
+}
+
+
+//////////////////////////////////////////////////////////////////////////
+
+FAnimNode_OrientationDriver::FAnimNode_OrientationDriver()
+{
+	RadialScaling = 0.25f;
+}
+
+void FAnimNode_OrientationDriver::GatherDebugData(FNodeDebugData& DebugData)
+{
+	FString DebugLine = DebugData.GetNodeName(this);
+
+	DebugData.AddDebugItem(DebugLine);
+
+	ComponentPose.GatherDebugData(DebugData);
+}
+
+void FAnimNode_OrientationDriver::EvaluateBoneTransforms(USkeletalMeshComponent* SkelComp, FCSPose<FCompactPose>& MeshBases, TArray<FBoneTransform>& OutCSBoneTransforms)
+{
+
+}
+
+
+
+void FAnimNode_OrientationDriver::UpdateMaxDistanceBetweenPoses()
+{
+	// Find largest distance between poses
+	MaxDistanceBetweenPoses = 0.f;
+
+	for (int32 I = 1; I < Poses.Num(); I++)
+	{
+		FQuat QuatI = Poses[I].PoseQuat;
+		for (int32 J = 0; J < I; J++)
+		{
+			FQuat QuatJ = Poses[J].PoseQuat;
+
+			float Dist = QuatI.AngularDistance(QuatJ);
+
+			MaxDistanceBetweenPoses = FMath::Max(Dist, MaxDistanceBetweenPoses);
+		}
+	}
+
+	bMaxDistanceUpToDate = true;
+}
+
+void FAnimNode_OrientationDriver::EvaluateComponentSpaceInternal(FComponentSpacePoseContext& Context)
+{
+	// Do nothing if no poses
+	if (Poses.Num() == 0)
+	{
+		return;
+	}
+
+	// Ensure max distance up to date
+	if (!bMaxDistanceUpToDate)
+	{
+		UpdateMaxDistanceBetweenPoses();
+	}
+
+	const FBoneContainer& BoneContainer = Context.Pose.GetPose().GetBoneContainer();
+
+	// Final set of driven params, from all interpolators
+	ResultParamSet.ClearParams();
+
+	// Get the current rotation of the source bone
+	const FTransform SourceCurrentBoneTransform = Context.Pose.GetLocalSpaceTransform(SourceBone.GetCompactPoseIndex(BoneContainer));
+	const FQuat BoneQuat = SourceCurrentBoneTransform.GetRotation();
+
+	// Array of 'weights' of current orientation to each pose
+	float TotalWeight = 0.f;
+
+	float GaussVarianceSqr = FMath::Square(RadialScaling * MaxDistanceBetweenPoses);
+
+	// Iterate over each pose, adding its contribution
+	for (FOrientationDriverPose& Pose : Poses)
+	{
+		// Find distance
+		Pose.PoseDistance = BoneQuat.AngularDistance(Pose.PoseQuat);
+
+		// Evaluate radial basis function to find weight
+		Pose.PoseWeight = FMath::Exp(-(Pose.PoseDistance * Pose.PoseDistance) / GaussVarianceSqr);
+
+		// Add weight to total
+		TotalWeight += Pose.PoseWeight;
+
+		// Add params for this pose, weighted by pose weight
+		ResultParamSet.AddParams(Pose.DrivenParams, Pose.PoseWeight);
+	}
+
+	// Only normalize and apply if we got some kind of weight
+	if (TotalWeight > KINDA_SMALL_NUMBER)
+	{
+		// Normalize params by rescaling by sum of weights
+		const float WeightScale = 1.f / TotalWeight;
+		ResultParamSet.ScaleAllParams(WeightScale);
+
+		// Also normalize each pose weight
+		for (FOrientationDriverPose& Pose : Poses)
+		{
+			Pose.PoseWeight *= WeightScale;
+		}
+
+		//	Morph target and Material parameter curves
+		USkeleton* Skeleton = Context.AnimInstanceProxy->GetSkeleton();
+
+		for (FOrientationDriverParam& Param : ResultParamSet.Params)
+		{
+			FSmartNameMapping::UID NameUID = Skeleton->GetUIDByName(USkeleton::AnimCurveMappingName, Param.ParameterName);
+			if (NameUID != FSmartNameMapping::MaxUID)
+			{
+				Context.Curve.Set(NameUID, Param.ParameterValue, EAnimCurveFlags::ACF_DrivesMorphTarget | EAnimCurveFlags::ACF_DrivesMaterial | EAnimCurveFlags::ACF_DrivesPose);
+			}
+		}
+	}
+}
+
+bool FAnimNode_OrientationDriver::IsValidToEvaluate(const USkeleton* Skeleton, const FBoneContainer& RequiredBones)
+{
+	// return true if at least one bone ref is valid
+	return SourceBone.IsValid(RequiredBones);
+}
+
+void FAnimNode_OrientationDriver::InitializeBoneReferences(const FBoneContainer& RequiredBones)
+{
+	// Init bone ref
+	SourceBone.Initialize(RequiredBones);
+}
\ No newline at end of file
diff --git a/Engine/Source/Runtime/AnimGraphRuntime/Public/BoneControllers/AnimNode_AnimDynamics.h b/Engine/Source/Runtime/AnimGraphRuntime/Public/BoneControllers/AnimNode_AnimDynamics.h
--- a/Engine/Source/Runtime/AnimGraphRuntime/Public/BoneControllers/AnimNode_AnimDynamics.h
+++ b/Engine/Source/Runtime/AnimGraphRuntime/Public/BoneControllers/AnimNode_AnimDynamics.h
@@ -27,7 +27,7 @@ enum class AnimPhysLinearConstraintType : uint8
 	Limited,
 };
 
-UENUM()
+UENUM(BlueprintType)
 enum class AnimPhysSimSpaceType : uint8
 {
 	Component UMETA(ToolTip = "Sim origin is the location/orientation of the skeletal mesh component."),
@@ -204,7 +204,7 @@ struct ANIMGRAPHRUNTIME_API FAnimNode_AnimDynamics : public FAnimNode_SkeletalCo
 	FAnimNode_AnimDynamics();
 	
 	/** The space used to run the simulation */
-	UPROPERTY(EditAnywhere, BlueprintReadWrite, Category = Setup)
+	UPROPERTY(EditAnywhere, BlueprintReadWrite, Category = Setup, meta = (PinHiddenByDefault))
 	AnimPhysSimSpaceType SimulationSpace;
 
 	/** When in BoneRelative sim space, the simulation will use this bone as the origin */
@@ -232,7 +232,7 @@ struct ANIMGRAPHRUNTIME_API FAnimNode_AnimDynamics : public FAnimNode_SkeletalCo
 	FVector LocalJointOffset;
 
 	/** Scale for gravity, higher values increase forces due to gravity */
-	UPROPERTY(EditAnywhere, BlueprintReadWrite, Category = Setup)
+	UPROPERTY(EditAnywhere, BlueprintReadWrite, Category = Setup, meta = (PinHiddenByDefault))
 	float GravityScale;
 
 	/** If true the body will attempt to spring back to its initial position */
@@ -244,11 +244,11 @@ struct ANIMGRAPHRUNTIME_API FAnimNode_AnimDynamics : public FAnimNode_SkeletalCo
 	bool bAngularSpring;
 
 	/** Spring constant to use when calculating linear springs, higher values mean a stronger spring.*/
-	UPROPERTY(EditAnywhere, BlueprintReadWrite, Category = Setup)
+	UPROPERTY(EditAnywhere, BlueprintReadWrite, Category = Setup, meta = (PinHiddenByDefault))
 	float LinearSpringConstant;
 
 	/** Spring constant to use when calculating angular springs, higher values mean a stronger spring */
-	UPROPERTY(EditAnywhere, BlueprintReadWrite, Category = Setup)
+	UPROPERTY(EditAnywhere, BlueprintReadWrite, Category = Setup, meta = (PinHiddenByDefault))
 	float AngularSpringConstant;
 
 	/** Whether or not wind is enabled for the bodies in this simulation */
@@ -264,7 +264,7 @@ struct ANIMGRAPHRUNTIME_API FAnimNode_AnimDynamics : public FAnimNode_SkeletalCo
 	bool bOverrideLinearDamping;
 
 	/** Overridden linear damping value */
-	UPROPERTY(EditAnywhere, AdvancedDisplay, BlueprintReadWrite, Category = Setup)
+	UPROPERTY(EditAnywhere, AdvancedDisplay, BlueprintReadWrite, Category = Setup, meta = (PinHiddenByDefault))
 	float LinearDampingOverride;
 
 	/** If true, the override value will be used for angular damping */
@@ -272,7 +272,7 @@ struct ANIMGRAPHRUNTIME_API FAnimNode_AnimDynamics : public FAnimNode_SkeletalCo
 	bool bOverrideAngularDamping;
 
 	/** Overridden angular damping value */
-	UPROPERTY(EditAnywhere, AdvancedDisplay, BlueprintReadWrite, Category = Setup)
+	UPROPERTY(EditAnywhere, AdvancedDisplay, BlueprintReadWrite, Category = Setup, meta = (PinHiddenByDefault))
 	float AngularDampingOverride;
 
 	/** If true we will perform physics update, otherwise skip - allows visualisation of the initial state of the bodies */
@@ -367,11 +367,15 @@ private:
 	FTransform GetBoneTransformInSimSpace(USkeletalMeshComponent* SkelComp, FCSPose<FCompactPose>& MeshBases, const FCompactPoseBoneIndex& BoneIndex);
 
 	// Given a transform in simulation space, convert it back to component space
-	FTransform GetComponentSpaceTransformFromSimSpace(USkeletalMeshComponent* SkelComp, FCSPose<FCompactPose>& MeshBases, const FCompactPoseBoneIndex& BoneIndex, const FTransform& InSimTransform);
+	FTransform GetComponentSpaceTransformFromSimSpace(AnimPhysSimSpaceType SimSpace, USkeletalMeshComponent* SkelComp, FCSPose<FCompactPose>& MeshBases, const FTransform& InSimTransform);
+	// Given a transform in component space, convert it to the current sim space
+	FTransform GetSimSpaceTransformFromComponentSpace(AnimPhysSimSpaceType SimSpace, USkeletalMeshComponent* SkelComp, FCSPose<FCompactPose>& MeshBases, const FTransform& InComponentTransform);
 
 	// Given a world-space vector, convert it into the current simulation space
 	FVector TransformWorldVectorToSimSpace(USkeletalMeshComponent* SkelComp, FCSPose<FCompactPose>& MeshBases, const FVector& InVec);
 
+	void ConvertSimulationSpace(USkeletalMeshComponent* SkelComp, FCSPose<FCompactPose>& MeshBases, AnimPhysSimSpaceType From, AnimPhysSimSpaceType To);
+
 	// We can't get clean bone positions unless we are in the evaluate step.
 	// Requesting an init or reinit sets this flag for us to pick up during evaluate
 	bool bRequiresInit;
@@ -394,6 +398,9 @@ private:
 	int32 MaxSubsteps;
 	//////////////////////////////////////////////////////////////////////////
 
+	// Cached sim space that we last used
+	AnimPhysSimSpaceType LastSimSpace;
+
 	// Active body list
 	TArray<FAnimPhysLinkedBody> Bodies;
 
diff --git a/Engine/Source/Runtime/AnimGraphRuntime/Public/BoneControllers/AnimNode_OrientationDriver.h b/Engine/Source/Runtime/AnimGraphRuntime/Public/BoneControllers/AnimNode_OrientationDriver.h
new file mode 100644
--- /dev/null
+++ b/Engine/Source/Runtime/AnimGraphRuntime/Public/BoneControllers/AnimNode_OrientationDriver.h
@@ -0,0 +1,110 @@
+// Copyright 1998-2016 Epic Games, Inc. All Rights Reserved.
+
+#pragma once
+
+#include "AnimNode_SkeletalControlBase.h"
+#include "AnimNode_OrientationDriver.generated.h"
+
+/** One named parameter being driven by the orientation of a bone. */
+USTRUCT()
+struct FOrientationDriverParam
+{
+	GENERATED_BODY()
+
+	UPROPERTY(EditAnywhere, Category = Param)
+	FName ParameterName;
+
+	UPROPERTY(EditAnywhere, Category = Param)
+	float ParameterValue;
+};
+
+/** One target pose for the bone, with parameters to drive bone approaches this orientation. */
+USTRUCT()
+struct FOrientationDriverPose
+{
+	GENERATED_BODY()
+
+	UPROPERTY(EditAnywhere, Category = Pose)
+	FQuat PoseQuat;
+
+	UPROPERTY(EditAnywhere, Category = Pose)
+	TArray<FOrientationDriverParam> DrivenParams;
+
+	/** Last weight calculated for this pose. */
+	float PoseWeight;
+
+	/** Last distance calculated from this pose */
+	float PoseDistance;
+
+	FOrientationDriverPose()
+	: PoseQuat(FQuat::Identity)
+	{}
+};
+
+/** Set of parameters being driven */
+struct FOrientationDriverParamSet
+{
+	/** Array of parameters that makes up a set */
+	TArray<FOrientationDriverParam> Params;
+
+	/** Increase value of a named parameter in a set. */
+	void AddParam(const FOrientationDriverParam& InParam, float InScale);
+
+	/** Add a set of params to another set */
+	void AddParams(const TArray<FOrientationDriverParam>& InParams, float InScale);
+
+	/** Scale all params in a set */
+	void ScaleAllParams(float InScale);
+
+	/** Clear all params */
+	void ClearParams();
+};
+
+/** RBF based orientation driver */
+USTRUCT()
+struct ANIMGRAPHRUNTIME_API FAnimNode_OrientationDriver : public FAnimNode_SkeletalControlBase
+{
+	GENERATED_BODY()
+
+	/** Bone to use for driving parameters based on its orientation */
+	UPROPERTY(EditAnywhere, Category = OrientationDriver)
+	FBoneReference SourceBone;
+
+	/** Set of target poses, which include parameters to drive. */
+	UPROPERTY(EditAnywhere, Category = OrientationDriver)
+	TArray<FOrientationDriverPose> Poses;
+
+	/** Scaling of radial basis, applied to max distance between poses */
+	UPROPERTY(EditAnywhere, Category = OrientationDriver, meta = (ClampMin = "0.01"))
+	float RadialScaling;
+
+	/** Previous set of interpolated params, kept for debugging */
+	FOrientationDriverParamSet ResultParamSet;
+
+	// FAnimNode_Base interface
+	virtual void GatherDebugData(FNodeDebugData& DebugData) override;
+	// End of FAnimNode_Base interface
+
+	FAnimNode_OrientationDriver();
+
+	// FAnimNode_SkeletalControlBase interface
+	virtual void EvaluateBoneTransforms(USkeletalMeshComponent* SkelComp, FCSPose<FCompactPose>& MeshBases, TArray<FBoneTransform>& OutBoneTransforms) override;
+	virtual void EvaluateComponentSpaceInternal(FComponentSpacePoseContext& Context);
+	virtual bool IsValidToEvaluate(const USkeleton* Skeleton, const FBoneContainer& RequiredBones) override;
+	// End of FAnimNode_SkeletalControlBase interface
+
+	/** Find max distance between any of the poses */
+	void UpdateMaxDistanceBetweenPoses();
+
+protected:
+
+	/** Is MaxDistanceBetweenPoses up to date */
+	bool bMaxDistanceUpToDate;
+
+	/** Maximum angular distance between poses in the Poses array */
+	float MaxDistanceBetweenPoses;
+
+	// FAnimNode_SkeletalControlBase protected interface
+	virtual void InitializeBoneReferences(const FBoneContainer& RequiredBones) override;
+	// End of FAnimNode_SkeletalControlBase protected interface
+};
\ No newline at end of file
diff --git a/Engine/Source/Runtime/Core/Public/Math/Quat.h b/Engine/Source/Runtime/Core/Public/Math/Quat.h
--- a/Engine/Source/Runtime/Core/Public/Math/Quat.h
+++ b/Engine/Source/Runtime/Core/Public/Math/Quat.h
@@ -377,6 +377,9 @@ public:
 	 */
 	FORCEINLINE FVector GetRotationAxis() const;
 
+	/** Find the angular distance between two rotation quaternions (in radians) */
+	FORCEINLINE float AngularDistance(const FQuat& Q) const;
+
 	/**
 	 * Serializes the vector compressed for e.g. network transmission.
 	 * @param Ar Archive to serialize to/ from.
@@ -947,6 +950,12 @@ FORCEINLINE FVector FQuat::GetRotationAxis() const
 	return FVector(1.f, 0.f, 0.f);
 }
 
+float FQuat::AngularDistance(const FQuat& Q) const
+{
+	float InnerProd = X*Q.X + Y*Q.Y + Z*Q.Z + W*Q.W;
+	return FMath::Acos((2 * InnerProd * InnerProd) - 1.f);
+}
+
 
 FORCEINLINE FVector FQuat::RotateVector(FVector V) const
 {	
diff --git a/Engine/Source/Runtime/Core/Public/UObject/FrameworkObjectVersion.h b/Engine/Source/Runtime/Core/Public/UObject/FrameworkObjectVersion.h
--- a/Engine/Source/Runtime/Core/Public/UObject/FrameworkObjectVersion.h
+++ b/Engine/Source/Runtime/Core/Public/UObject/FrameworkObjectVersion.h
@@ -35,6 +35,9 @@ struct CORE_API FFrameworkObjectVersion
 		// Add Reference Skeleton to Rig
 		AddSourceReferenceSkeletonToRig, 
 
+		// Refactor ConstraintInstance so that we have an easy way to swap behavior paramters
+		ConstraintInstanceBehaviorParameters,
+
 		// -----<new versions can be added above this line>-------------------------------------------------
 		VersionPlusOne,
 		LatestVersion = VersionPlusOne - 1
diff --git a/Engine/Source/Runtime/Engine/Classes/AI/Navigation/NavigationSystem.h b/Engine/Source/Runtime/Engine/Classes/AI/Navigation/NavigationSystem.h
--- a/Engine/Source/Runtime/Engine/Classes/AI/Navigation/NavigationSystem.h
+++ b/Engine/Source/Runtime/Engine/Classes/AI/Navigation/NavigationSystem.h
@@ -509,6 +509,7 @@ public:
 
 	bool ShouldGenerateNavigationEverywhere() const { return bWholeWorldNavigable; }
 
+	bool ShouldAllowClientSideNavigation() const { return bAllowClientSideNavigation; }
 	virtual bool ShouldLoadNavigationOnClient(ANavigationData* NavData = nullptr) const { return bAllowClientSideNavigation; }
 
 	FBox GetWorldBounds() const;
diff --git a/Engine/Source/Runtime/Engine/Classes/Camera/CameraTypes.h b/Engine/Source/Runtime/Engine/Classes/Camera/CameraTypes.h
--- a/Engine/Source/Runtime/Engine/Classes/Camera/CameraTypes.h
+++ b/Engine/Source/Runtime/Engine/Classes/Camera/CameraTypes.h
@@ -4,6 +4,7 @@
 
 #include "Engine/Scene.h"
 #include "Engine/EngineTypes.h"
+#include "EngineDefines.h"
 #include "CameraTypes.generated.h"
 
 //@TODO: Document
diff --git a/Engine/Source/Runtime/Engine/Classes/Components/SkeletalMeshComponent.h b/Engine/Source/Runtime/Engine/Classes/Components/SkeletalMeshComponent.h
--- a/Engine/Source/Runtime/Engine/Classes/Components/SkeletalMeshComponent.h
+++ b/Engine/Source/Runtime/Engine/Classes/Components/SkeletalMeshComponent.h
@@ -986,6 +986,30 @@ public:
 	virtual float GetMass() const override;
 	virtual float CalculateMass(FName BoneName = NAME_None) override;
 	virtual bool DoCustomNavigableGeometryExport(FNavigableGeometryExport& GeomExport) const override;
+
+	/**
+	*	Add a force to all rigid bodies below.
+	*   This is like a 'thruster'. Good for adding a burst over some (non zero) time. Should be called every frame for the duration of the force.
+	*
+	*	@param	Force		 Force vector to apply. Magnitude indicates strength of force.
+	*	@param	BoneName	 If a SkeletalMeshComponent, name of body to apply force to. 'None' indicates root body.
+	*   @param  bAccelChange If true, Force is taken as a change in acceleration instead of a physical force (i.e. mass will have no affect).
+	*   @param  bIncludeSelf If false, Force is only applied to bodies below but not given bone name.
+	*/
+	UFUNCTION(BlueprintCallable, Category = "Physics")
+	virtual void AddForceToAllBodiesBelow(FVector Force, FName BoneName = NAME_None, bool bAccelChange = false, bool bIncludeSelf = true);
+
+	/**
+	*	Add impulse to all single rigid bodies below. Good for one time instant burst.
+	*
+	*	@param	Impulse		Magnitude and direction of impulse to apply.
+	*	@param	BoneName	If a SkeletalMeshComponent, name of body to apply impulse to. 'None' indicates root body.
+	*	@param	bVelChange	If true, the Strength is taken as a change in velocity instead of an impulse (ie. mass will have no affect).
+	*	@param bIncludeSelf If false, Force is only applied to bodies below but not given bone name.
+	*/
+	UFUNCTION(BlueprintCallable, Category = "Physics")
+	virtual void AddImpulseToAllBodiesBelow(FVector Impulse, FName BoneName = NAME_None, bool bVelChange = false, bool bIncludeSelf = true);
+
 #if WITH_EDITOR
 	virtual bool ComponentIsTouchingSelectionBox(const FBox& InSelBBox, const FEngineShowFlags& ShowFlags, const bool bConsiderOnlyBSP, const bool bMustEncompassEntireComponent) const override;
 	virtual bool ComponentIsTouchingSelectionFrustum(const FConvexVolume& InFrustum, const FEngineShowFlags& ShowFlags, const bool bConsiderOnlyBSP, const bool bMustEncompassEntireComponent) const override;
@@ -1094,7 +1118,7 @@ public:
 
 	/** Set all of the bones below passed in bone to be simulated */
 	UFUNCTION(BlueprintCallable, Category="Components|SkeletalMesh")
-	void SetAllBodiesBelowSimulatePhysics(const FName& InBoneName, bool bNewSimulate );
+	void SetAllBodiesBelowSimulatePhysics(const FName& InBoneName, bool bNewSimulate, bool bIncludeSelf = true );
 
 	/** Allows you to reset bodies Simulate state based on where bUsePhysics is set to true in the BodySetup. */
 	UFUNCTION(BlueprintCallable, Category="Components|SkeletalMesh")
@@ -1105,7 +1129,7 @@ public:
 
 	/** Set all of the bones below passed in bone to be simulated */
 	UFUNCTION(BlueprintCallable, Category="Components|SkeletalMesh")
-	void SetAllBodiesBelowPhysicsBlendWeight(const FName& InBoneName, float PhysicsBlendWeight, bool bSkipCustomPhysicsType = false );
+	void SetAllBodiesBelowPhysicsBlendWeight(const FName& InBoneName, float PhysicsBlendWeight, bool bSkipCustomPhysicsType = false, bool bIncludeSelf = true );
 
 	/** Accumulate AddPhysicsBlendWeight to physics blendweight for all of the bones below passed in bone to be simulated */
 	UFUNCTION(BlueprintCallable, Category="Components|SkeletalMesh")
@@ -1123,6 +1147,10 @@ public:
 	UFUNCTION(BlueprintCallable, Category = "Components|SkeletalMesh")
 	void SetAllMotorsAngularDriveParams(float InSpring, float InDamping, float InForceLimit, bool bSkipCustomPhysicsType = false);
 
+	/** Sets the constraint profile properties (limits, motors, etc...) to match the constraint profile as defined in the physics asset. If profile name is not found the joint is remains untouched*/
+	UFUNCTION(BlueprintCallable, Category = "Physics|Components|PhysicsConstraint")
+	void SetConstraintProfile(FName JointName, FName ProfileName);
+
 	/** Enable or Disable AngularPositionDrive based on a list of bone names */
 	void SetNamedMotorsAngularPositionDrive(bool bEnableSwingDrive, bool bEnableTwistDrive, const TArray<FName>& BoneNames, bool bSetOtherBodiesToComplement = false);
 
@@ -1131,6 +1159,9 @@ public:
 
 	void GetWeldedBodies(TArray<FBodyInstance*> & OutWeldedBodies, TArray<FName> & OutChildrenLabels) override;
 
+	/** Iterates over all bodies below and executes Func. Returns number of bodies found */
+	int32 ForEachBodyBelow(FName BoneName, bool bIncludeSelf, bool bSkipCustomType, TFunctionRef<void(FBodyInstance*)> Func);
+
 	/** 
 	 * Change whether to force mesh into ref pose (and use cheaper vertex shader) 
 	 *
diff --git a/Engine/Source/Runtime/Engine/Classes/Engine/Engine.h b/Engine/Source/Runtime/Engine/Classes/Engine/Engine.h
--- a/Engine/Source/Runtime/Engine/Classes/Engine/Engine.h
+++ b/Engine/Source/Runtime/Engine/Classes/Engine/Engine.h
@@ -890,19 +890,22 @@ public:
 	UPROPERTY()
 	class UMaterialInstanceDynamic* ConstraintLimitMaterialX;
 
+	UPROPERTY()
+	class UMaterialInstanceDynamic* ConstraintLimitMaterialXAxis;
+
 	UPROPERTY()
 	class UMaterialInstanceDynamic* ConstraintLimitMaterialY;
+	UPROPERTY()
+	class UMaterialInstanceDynamic* ConstraintLimitMaterialYAxis;
 
 	UPROPERTY()
 	class UMaterialInstanceDynamic* ConstraintLimitMaterialZ;
+	UPROPERTY()
+	class UMaterialInstanceDynamic* ConstraintLimitMaterialZAxis;
 
 	UPROPERTY()
 	class UMaterialInstanceDynamic* ConstraintLimitMaterialPrismatic;
 
-	/** @todo document */
-	UPROPERTY(globalconfig)
-	FStringAssetReference ConstraintLimitMaterialName;
-
 	/** Material that renders a message about lightmap settings being invalid. */
 	UPROPERTY()
 	class UMaterial* InvalidLightmapSettingsMaterial;
diff --git a/Engine/Source/Runtime/Engine/Classes/Engine/SkeletalMesh.h b/Engine/Source/Runtime/Engine/Classes/Engine/SkeletalMesh.h
--- a/Engine/Source/Runtime/Engine/Classes/Engine/SkeletalMesh.h
+++ b/Engine/Source/Runtime/Engine/Classes/Engine/SkeletalMesh.h
@@ -310,38 +310,6 @@ struct FSkeletalMeshLODInfo
 
 };
 
-USTRUCT()
-struct FMorphTargetMap
-{
-	GENERATED_USTRUCT_BODY()
-
-	/** The bone to mirror. */
-	UPROPERTY(EditAnywhere, Category=MorphTargetMap)
-	FName Name;
-
-	/** Axis the bone is mirrored across. */
-	UPROPERTY(EditAnywhere, Category=MorphTargetMap)
-	UMorphTarget* MorphTarget;
-
-
-	FMorphTargetMap()
-		: Name(NAME_None)
-		, MorphTarget(NULL)
-	{
-	}
-
-	FMorphTargetMap( FName InName, UMorphTarget * InMorphTarget )
-		: Name(InName)
-		, MorphTarget(InMorphTarget)
-	{
-	}
-
-	bool operator== (const FMorphTargetMap& Other) const
-	{
-		return (Name==Other.Name && MorphTarget == Other.MorphTarget);
-	}
-};
-
 /** 
  * constrain Coefficients - max distance, collisionSphere radius, collision sphere distance 
  */
