Fix non-unity

diff --git a/Engine/Source/Runtime/AnimGraphRuntime/Public/BoneControllers/AnimNode_OrientationDriver.h b/Engine/Source/Runtime/AnimGraphRuntime/Public/BoneControllers/AnimNode_OrientationDriver.h
--- a/Engine/Source/Runtime/AnimGraphRuntime/Public/BoneControllers/AnimNode_OrientationDriver.h
+++ b/Engine/Source/Runtime/AnimGraphRuntime/Public/BoneControllers/AnimNode_OrientationDriver.h
@@ -5,6 +5,8 @@
 #include "AnimNode_SkeletalControlBase.h"
 #include "AnimNode_OrientationDriver.generated.h"
 
+class UPoseAsset;
+
 /** One named parameter being driven by the orientation of a bone. */
 struct FOrientationDriverParam
 {
