DUPE MERGE: notify issue with sequencer(CL 3628826)

diff --git a/Engine/Source/Runtime/AnimGraphRuntime/Private/AnimSequencerInstanceProxy.cpp b/Engine/Source/Runtime/AnimGraphRuntime/Private/AnimSequencerInstanceProxy.cpp
--- a/Engine/Source/Runtime/AnimGraphRuntime/Private/AnimSequencerInstanceProxy.cpp
+++ b/Engine/Source/Runtime/AnimGraphRuntime/Private/AnimSequencerInstanceProxy.cpp
@@ -7,6 +7,8 @@ void FAnimSequencerInstanceProxy::Initialize(UAnimInstance* InAnimInstance)
 {
 	FAnimInstanceProxy::Initialize(InAnimInstance);
 	ConstructNodes();
+
+	UpdateCounter.Reset();
 }
 
 bool FAnimSequencerInstanceProxy::Evaluate(FPoseContext& Output)
@@ -18,6 +20,8 @@ bool FAnimSequencerInstanceProxy::Evaluate(FPoseContext& Output)
 
 void FAnimSequencerInstanceProxy::UpdateAnimationNode(float DeltaSeconds)
 {
+	UpdateCounter.Increment();
+
 	SequencerRootNode.Update_AnyThread(FAnimationUpdateContext(this, DeltaSeconds));
 }
 
