Merging using Dev-AnimPhys<->Dev-Anim

CCDIK node

diff --git a/Engine/Source/Editor/AnimGraph/Classes/AnimGraphNode_CCDIK.h b/Engine/Source/Editor/AnimGraph/Classes/AnimGraphNode_CCDIK.h
new file mode 100644
--- /dev/null
+++ b/Engine/Source/Editor/AnimGraph/Classes/AnimGraphNode_CCDIK.h
@@ -0,0 +1,40 @@
+// Copyright 1998-2016 Epic Games, Inc. All Rights Reserved.
+
+#pragma once
+#include "AnimGraphNode_SkeletalControlBase.h"
+#include "BoneControllers/AnimNode_CCDIK.h"
+#include "EdGraph/EdGraphNodeUtils.h"
+#include "AnimGraphNode_CCDIK.generated.h"
+
+// Editor node for CCDIK IK skeletal controller
+UCLASS(MinimalAPI, Experimental)
+class UAnimGraphNode_CCDIK : public UAnimGraphNode_SkeletalControlBase
+{
+	GENERATED_UCLASS_BODY()
+
+	UPROPERTY(EditAnywhere, Category = Settings)
+	FAnimNode_CCDIK Node;
+
+public:
+	// UEdGraphNode interface
+	virtual FText GetNodeTitle(ENodeTitleType::Type TitleType) const override;
+	// End of UEdGraphNode interface
+
+	// UAnimGraphNode_Base interface
+	virtual void CopyNodeDataToPreviewNode(FAnimNode_Base* AnimNode) override;
+	virtual FEditorModeID GetEditorMode() const override;
+	// End of UAnimGraphNode_Base interface
+
+	// UObject interface
+	virtual void PostEditChangeProperty(struct FPropertyChangedEvent& PropertyChangedEvent) override;
+
+protected:
+	// UAnimGraphNode_SkeletalControlBase interface
+	virtual FText GetControllerDescription() const override;
+	virtual const FAnimNode_SkeletalControlBase* GetNode() const override { return &Node; }
+	// End of UAnimGraphNode_SkeletalControlBase interface
+
+private:
+	/** Constructing FText strings can be costly, so we cache the node's title */
+	FNodeTitleTextTable CachedNodeTitles;
+};
\ No newline at end of file
diff --git a/Engine/Source/Editor/AnimGraph/Private/AnimGraphModule.cpp b/Engine/Source/Editor/AnimGraph/Private/AnimGraphModule.cpp
--- a/Engine/Source/Editor/AnimGraph/Private/AnimGraphModule.cpp
+++ b/Engine/Source/Editor/AnimGraph/Private/AnimGraphModule.cpp
@@ -18,6 +18,7 @@
 #include "EditModes/PoseDriverEditMode.h"
 #include "EditModes/SplineIKEditMode.h"
 #include "EditModes/LookAtEditMode.h"
+#include "EditModes/CCDIKEditMode.h"
 
 IMPLEMENT_MODULE(FAnimGraphModule, AnimGraph);
 
@@ -36,6 +37,7 @@ void FAnimGraphModule::StartupModule()
 	FEditorModeRegistry::Get().RegisterMode<FPoseDriverEditMode>(AnimNodeEditModes::PoseDriver, LOCTEXT("PoseDriverEditMode", "PoseDriver"), FSlateIcon(), false);
 	FEditorModeRegistry::Get().RegisterMode<FSplineIKEditMode>(AnimNodeEditModes::SplineIK, LOCTEXT("SplineIKEditMode", "Spline IK"), FSlateIcon(), false);
 	FEditorModeRegistry::Get().RegisterMode<FLookAtEditMode>(AnimNodeEditModes::LookAt, LOCTEXT("LookAtEditMode", "LookAt"), FSlateIcon(), false);
+	FEditorModeRegistry::Get().RegisterMode<FCCDIKEditMode>(AnimNodeEditModes::CCDIK, LOCTEXT("CCDIKEditMode", "CCD IK"), FSlateIcon(), false);
 
 	// Register details customization
 	FPropertyEditorModule& PropertyModule = FModuleManager::LoadModuleChecked<FPropertyEditorModule>("PropertyEditor");
@@ -45,6 +47,7 @@ void FAnimGraphModule::StartupModule()
 void FAnimGraphModule::ShutdownModule()
 {
 	// Unregister the editor modes
+	FEditorModeRegistry::Get().UnregisterMode(AnimNodeEditModes::CCDIK);
 	FEditorModeRegistry::Get().UnregisterMode(AnimNodeEditModes::SplineIK);
 	FEditorModeRegistry::Get().UnregisterMode(AnimNodeEditModes::PoseDriver);
 	FEditorModeRegistry::Get().UnregisterMode(AnimNodeEditModes::Fabrik);
diff --git a/Engine/Source/Editor/AnimGraph/Private/AnimGraphNode_CCDIK.cpp b/Engine/Source/Editor/AnimGraph/Private/AnimGraphNode_CCDIK.cpp
new file mode 100644
--- /dev/null
+++ b/Engine/Source/Editor/AnimGraph/Private/AnimGraphNode_CCDIK.cpp
@@ -0,0 +1,101 @@
+// Copyright 1998-2016 Epic Games, Inc. All Rights Reserved.
+
+#include "AnimGraphNode_CCDIK.h"
+#include "Animation/AnimInstance.h"
+#include "AnimNodeEditModes.h"
+
+/////////////////////////////////////////////////////
+// UAnimGraphNode_CCDIK 
+
+#define LOCTEXT_NAMESPACE "A3Nodes"
+
+UAnimGraphNode_CCDIK::UAnimGraphNode_CCDIK(const FObjectInitializer& ObjectInitializer)
+: Super(ObjectInitializer)
+{
+}
+
+FText UAnimGraphNode_CCDIK::GetControllerDescription() const
+{
+	return LOCTEXT("CCDIK", "CCDIK");
+}
+
+FText UAnimGraphNode_CCDIK::GetNodeTitle(ENodeTitleType::Type TitleType) const
+{
+	if ((TitleType == ENodeTitleType::ListView || TitleType == ENodeTitleType::MenuTitle))
+	{
+		return GetControllerDescription();
+	}
+	// @TODO: the bone can be altered in the property editor, so we have to 
+	//        choose to mark this dirty when that happens for this to properly work
+	else //if (!CachedNodeTitles.IsTitleCached(TitleType, this))
+	{
+		FFormatNamedArguments Args;
+		Args.Add(TEXT("ControllerDescription"), GetControllerDescription());
+		Args.Add(TEXT("RootBoneName"), FText::FromName(Node.RootBone.BoneName));
+		Args.Add(TEXT("TipBoneName"), FText::FromName(Node.TipBone.BoneName));
+
+		// FText::Format() is slow, so we cache this to save on performance
+		if (TitleType == ENodeTitleType::ListView || TitleType == ENodeTitleType::MenuTitle)
+		{
+			CachedNodeTitles.SetCachedTitle(TitleType, FText::Format(LOCTEXT("AnimGraphNode_CCDIKBone_ListTitle", "{ControllerDescription} - Root: {RootBoneName}, Tip: {TipBoneName} "), Args), this);
+		}
+		else
+		{
+			CachedNodeTitles.SetCachedTitle(TitleType, FText::Format(LOCTEXT("AnimGraphNode_CCDIKBone_Title", "{ControllerDescription}\nRoot: {RootBoneName} - Tip: {TipBoneName} "), Args), this);
+		}
+	}
+	return CachedNodeTitles[TitleType];
+}
+
+void UAnimGraphNode_CCDIK::CopyNodeDataToPreviewNode(FAnimNode_Base* InPreviewNode)
+{
+	FAnimNode_CCDIK* CCDIK = static_cast<FAnimNode_CCDIK*>(InPreviewNode);
+
+	// copies Pin values from the internal node to get data which are not compiled yet
+	CCDIK->EffectorLocation = Node.EffectorLocation;
+}
+
+FEditorModeID UAnimGraphNode_CCDIK::GetEditorMode() const
+{
+	return AnimNodeEditModes::CCDIK;
+}
+
+void UAnimGraphNode_CCDIK::PostEditChangeProperty(struct FPropertyChangedEvent& PropertyChangedEvent)
+{
+	Super::PostEditChangeProperty(PropertyChangedEvent);
+
+	FName PropertyName = (PropertyChangedEvent.Property != NULL) ? PropertyChangedEvent.Property->GetFName() : NAME_None;
+	if (PropertyName == GET_MEMBER_NAME_CHECKED(FBoneReference, BoneName))
+	{
+		USkeleton* Skeleton = GetAnimBlueprint()->TargetSkeleton;
+		if (Node.TipBone.BoneName != NAME_None && Node.RootBone.BoneName != NAME_None)
+		{
+			const FReferenceSkeleton& RefSkeleton = Skeleton->GetReferenceSkeleton();
+			const int32 TipBoneIndex = RefSkeleton.FindBoneIndex(Node.TipBone.BoneName);
+			const int32 RootBoneIndex = RefSkeleton.FindBoneIndex(Node.RootBone.BoneName);
+
+			if (TipBoneIndex != INDEX_NONE && RootBoneIndex != INDEX_NONE)
+			{
+				const int32 Depth = RefSkeleton.GetDepthBetweenBones(TipBoneIndex, RootBoneIndex);
+				if (Depth >= 0)
+				{
+					Node.ResizeRotationLimitPerJoints(Depth + 1);
+				}
+				else
+				{
+					Node.ResizeRotationLimitPerJoints(0);
+				}
+			}
+			else
+			{
+				Node.ResizeRotationLimitPerJoints(0);
+			}
+		}
+		else
+		{
+			Node.ResizeRotationLimitPerJoints(0);
+		}
+	}
+}
+
+#undef LOCTEXT_NAMESPACE
\ No newline at end of file
diff --git a/Engine/Source/Editor/AnimGraph/Private/EditModes/CCDIKEditMode.cpp b/Engine/Source/Editor/AnimGraph/Private/EditModes/CCDIKEditMode.cpp
new file mode 100644
--- /dev/null
+++ b/Engine/Source/Editor/AnimGraph/Private/EditModes/CCDIKEditMode.cpp
@@ -0,0 +1,78 @@
+// Copyright 1998-2016 Epic Games, Inc. All Rights Reserved.
+
+#include "CCDIKEditMode.h"
+#include "AnimGraphNode_CCDIK.h"
+#include "IPersonaPreviewScene.h"
+#include "Animation/DebugSkelMeshComponent.h"
+
+void FCCDIKEditMode::EnterMode(class UAnimGraphNode_Base* InEditorNode, struct FAnimNode_Base* InRuntimeNode)
+{
+	RuntimeNode = static_cast<FAnimNode_CCDIK*>(InRuntimeNode);
+	GraphNode = CastChecked<UAnimGraphNode_CCDIK>(InEditorNode);
+
+	FAnimNodeEditMode::EnterMode(InEditorNode, InRuntimeNode);
+}
+
+void FCCDIKEditMode::ExitMode()
+{
+	RuntimeNode = nullptr;
+	GraphNode = nullptr;
+
+	FAnimNodeEditMode::ExitMode();
+}
+
+FVector FCCDIKEditMode::GetWidgetLocation() const
+{
+	EBoneControlSpace Space = RuntimeNode->EffectorLocationSpace;
+	FVector Location = RuntimeNode->EffectorLocation;
+	FBoneSocketTarget Target = RuntimeNode->EffectorTarget;
+
+	USkeletalMeshComponent* SkelComp = GetAnimPreviewScene().GetPreviewMeshComponent();
+	return ConvertWidgetLocation(SkelComp, RuntimeNode->ForwardedPose, Target, Location, Space);
+}
+
+FWidget::EWidgetMode FCCDIKEditMode::GetWidgetMode() const
+{
+	USkeletalMeshComponent* SkelComp = GetAnimPreviewScene().GetPreviewMeshComponent();
+	int32 TipBoneIndex = SkelComp->GetBoneIndex(RuntimeNode->TipBone.BoneName);
+	int32 RootBoneIndex = SkelComp->GetBoneIndex(RuntimeNode->RootBone.BoneName);
+
+	if (TipBoneIndex!= INDEX_NONE && RootBoneIndex != INDEX_NONE)
+	{
+		return FWidget::WM_Translate;
+	}
+
+	return FWidget::WM_None;
+}
+
+void FCCDIKEditMode::DoTranslation(FVector& InTranslation)
+{
+	USkeletalMeshComponent* SkelComp = GetAnimPreviewScene().GetPreviewMeshComponent();
+
+	FVector Offset = ConvertCSVectorToBoneSpace(SkelComp, InTranslation, RuntimeNode->ForwardedPose, RuntimeNode->EffectorTarget, GraphNode->Node.EffectorLocationSpace);
+
+	RuntimeNode->EffectorLocation += Offset;
+	GraphNode->Node.EffectorLocation = RuntimeNode->EffectorLocation;
+	GraphNode->SetDefaultValue(GET_MEMBER_NAME_STRING_CHECKED(FAnimNode_CCDIK, EffectorLocation), RuntimeNode->EffectorLocation);
+}
+
+void FCCDIKEditMode::Render(const FSceneView* View, FViewport* Viewport, FPrimitiveDrawInterface* PDI)
+{
+#if !(UE_BUILD_SHIPPING || UE_BUILD_TEST)
+	// draw line from root bone to tip bone if available
+	if (RuntimeNode && RuntimeNode->DebugLines.Num() > 0)
+	{
+		USkeletalMeshComponent* SkelComp = GetAnimPreviewScene().GetPreviewMeshComponent();
+		FTransform CompToWorld = SkelComp->GetComponentToWorld();
+
+		// no component space
+		for (int32 Index = 1; Index < RuntimeNode->DebugLines.Num(); ++Index)
+		{
+			FVector Start = CompToWorld.TransformPosition(RuntimeNode->DebugLines[Index - 1]);
+			FVector End = CompToWorld.TransformPosition(RuntimeNode->DebugLines[Index]);
+
+			PDI->DrawLine(Start, End, FLinearColor::Red, SDPG_Foreground);
+		}
+	}
+#endif // #if !(UE_BUILD_SHIPPING || UE_BUILD_TEST)
+}
\ No newline at end of file
diff --git a/Engine/Source/Editor/AnimGraph/Private/EditModes/CCDIKEditMode.h b/Engine/Source/Editor/AnimGraph/Private/EditModes/CCDIKEditMode.h
new file mode 100644
--- /dev/null
+++ b/Engine/Source/Editor/AnimGraph/Private/EditModes/CCDIKEditMode.h
@@ -0,0 +1,20 @@
+// Copyright 1998-2016 Epic Games, Inc. All Rights Reserved.
+
+#pragma once
+
+#include "AnimNodeEditMode.h"
+
+class FCCDIKEditMode : public FAnimNodeEditMode
+{
+public:
+	/** IAnimNodeEditMode interface */
+	virtual void EnterMode(class UAnimGraphNode_Base* InEditorNode, struct FAnimNode_Base* InRuntimeNode) override;
+	virtual void ExitMode() override;
+	virtual FVector GetWidgetLocation() const override;
+	virtual FWidget::EWidgetMode GetWidgetMode() const override;
+	virtual void DoTranslation(FVector& InTranslation) override;
+	virtual void Render(const FSceneView* View, FViewport* Viewport, FPrimitiveDrawInterface* PDI) override;
+private:
+	struct FAnimNode_CCDIK* RuntimeNode;
+	class UAnimGraphNode_CCDIK* GraphNode;
+};
\ No newline at end of file
diff --git a/Engine/Source/Runtime/AnimGraphRuntime/Private/BoneControllers/AnimNode_CCDIK.cpp b/Engine/Source/Runtime/AnimGraphRuntime/Private/BoneControllers/AnimNode_CCDIK.cpp
new file mode 100644
--- /dev/null
+++ b/Engine/Source/Runtime/AnimGraphRuntime/Private/BoneControllers/AnimNode_CCDIK.cpp
@@ -0,0 +1,318 @@
+// Copyright 1998-2016 Epic Games, Inc. All Rights Reserved.
+
+#include "BoneControllers/AnimNode_CCDIK.h"
+#include "Animation/AnimTypes.h"
+#include "AnimationRuntime.h"
+#include "DrawDebugHelpers.h"
+#include "Animation/AnimInstanceProxy.h"
+
+PRAGMA_DISABLE_OPTIMIZATION
+/////////////////////////////////////////////////////
+// AnimNode_CCDIK
+// Implementation of the CCDIK IK Algorithm
+
+FAnimNode_CCDIK::FAnimNode_CCDIK()
+	: EffectorLocation(FVector::ZeroVector)
+	, EffectorLocationSpace(BCS_ComponentSpace)
+	, Precision(1.f)
+	, MaxIterations(10)
+	, bStartFromTail(true)
+	, bEnableRotationLimit(false)
+{
+}
+
+FVector FAnimNode_CCDIK::GetCurrentLocation(FCSPose<FCompactPose>& MeshBases, const FCompactPoseBoneIndex& BoneIndex)
+{
+	return MeshBases.GetComponentSpaceTransform(BoneIndex).GetLocation();
+}
+
+FTransform FAnimNode_CCDIK::GetTargetTransform(const FTransform& InComponentTransform, FCSPose<FCompactPose>& MeshBases, FBoneSocketTarget& InTarget, EBoneControlSpace Space, const FVector& InOffset)
+{
+	FTransform OutTransform;
+	if (Space == BCS_BoneSpace)
+	{
+		OutTransform = InTarget.GetTargetTransform(InOffset, MeshBases, InComponentTransform);
+	}
+	else
+	{
+		// parent bone space still goes through this way
+		// if your target is socket, it will try find parents of joint that socket belongs to
+		OutTransform.SetLocation(InOffset);
+		FAnimationRuntime::ConvertBoneSpaceTransformToCS(InComponentTransform, MeshBases, OutTransform, InTarget.GetCompactPoseBoneIndex(), Space);
+	}
+
+	return OutTransform;
+}
+
+void FAnimNode_CCDIK::EvaluateSkeletalControl_AnyThread(FComponentSpacePoseContext& Output, TArray<FBoneTransform>& OutBoneTransforms)
+{
+	const FBoneContainer& BoneContainer = Output.Pose.GetPose().GetBoneContainer();
+
+	// Update EffectorLocation if it is based off a bone position
+	FTransform CSEffectorTransform = GetTargetTransform(Output.AnimInstanceProxy->GetComponentTransform(), Output.Pose, EffectorTarget, EffectorLocationSpace, EffectorLocation);
+	FVector const CSEffectorLocation = CSEffectorTransform.GetLocation();
+
+	// Gather all bone indices between root and tip.
+	TArray<FCompactPoseBoneIndex> BoneIndices;
+	
+	{
+		const FCompactPoseBoneIndex RootIndex = RootBone.GetCompactPoseIndex(BoneContainer);
+		FCompactPoseBoneIndex BoneIndex = TipBone.GetCompactPoseIndex(BoneContainer);
+		do
+		{
+			BoneIndices.Insert(BoneIndex, 0);
+			BoneIndex = Output.Pose.GetPose().GetParentBoneIndex(BoneIndex);
+		} while (BoneIndex != RootIndex);
+		BoneIndices.Insert(BoneIndex, 0);
+	}
+
+	// Gather transforms
+	int32 const NumTransforms = BoneIndices.Num();
+	OutBoneTransforms.AddUninitialized(NumTransforms);
+
+	// Gather chain links. These are non zero length bones.
+	TArray<CCDIKChainLink> Chain;
+	Chain.Reserve(NumTransforms);
+	// Start with Root Bone
+	{
+		const FCompactPoseBoneIndex& RootBoneIndex = BoneIndices[0];
+		const FTransform& LocalTransform = Output.Pose.GetLocalSpaceTransform(RootBoneIndex);
+		const FTransform& BoneCSTransform = Output.Pose.GetComponentSpaceTransform(RootBoneIndex);
+
+		OutBoneTransforms[0] = FBoneTransform(RootBoneIndex, BoneCSTransform);
+		Chain.Add(CCDIKChainLink(BoneCSTransform, LocalTransform, RootBoneIndex, 0));
+	}
+
+	// Go through remaining transforms
+	for (int32 TransformIndex = 1; TransformIndex < NumTransforms; TransformIndex++)
+	{
+		const FCompactPoseBoneIndex& BoneIndex = BoneIndices[TransformIndex];
+
+		const FTransform& LocalTransform = Output.Pose.GetLocalSpaceTransform(BoneIndex);
+		const FTransform& BoneCSTransform = Output.Pose.GetComponentSpaceTransform(BoneIndex);
+		FVector const BoneCSPosition = BoneCSTransform.GetLocation();
+
+		OutBoneTransforms[TransformIndex] = FBoneTransform(BoneIndex, BoneCSTransform);
+
+		// Calculate the combined length of this segment of skeleton
+		float const BoneLength = FVector::Dist(BoneCSPosition, OutBoneTransforms[TransformIndex - 1].Transform.GetLocation());
+
+		if (!FMath::IsNearlyZero(BoneLength))
+		{
+			Chain.Add(CCDIKChainLink(BoneCSTransform, LocalTransform, BoneIndex, TransformIndex));
+		}
+		else
+		{
+			// Mark this transform as a zero length child of the last link.
+			// It will inherit position and delta rotation from parent link.
+			CCDIKChainLink & ParentLink = Chain[Chain.Num() - 1];
+			ParentLink.ChildZeroLengthTransformIndices.Add(TransformIndex);
+		}
+	}
+
+	bool bBoneLocationUpdated = false;
+	int32 const NumChainLinks = BoneIndices.Num();
+
+	// iterate
+	{
+		int32 const TipBoneLinkIndex = NumChainLinks - 1;
+
+		// @todo optimize locally if no update, stop?
+		bool bLocalUpdated = false;
+		// check how far
+		const FVector TargetPos = CSEffectorLocation;
+		FVector TipPos = Chain[TipBoneLinkIndex].Transform.GetLocation();
+		float Distance = FVector::Dist(TargetPos, TipPos);
+		int32 IterationCount = 0;
+		while ((Distance > Precision) && (IterationCount++ < MaxIterations))
+		{
+			// iterate from tip to root
+			if (bStartFromTail)
+			{
+				for (int32 LinkIndex = TipBoneLinkIndex - 1; LinkIndex > 0; --LinkIndex)
+				{
+					bLocalUpdated |= UpdateChainLink(Chain, LinkIndex, TargetPos);
+				}
+			}
+			else
+			{
+				for (int32 LinkIndex = 1; LinkIndex < TipBoneLinkIndex; ++LinkIndex)
+				{
+					bLocalUpdated |= UpdateChainLink(Chain, LinkIndex, TargetPos);
+				}
+			}
+
+			Distance = FVector::Dist(Chain[TipBoneLinkIndex].Transform.GetLocation(), CSEffectorLocation);
+
+			bBoneLocationUpdated |= bLocalUpdated;
+
+			// no more update in this iteration
+			if (!bLocalUpdated)
+			{
+				break;
+			}
+		}
+	}
+
+	// If we moved some bones, update bone transforms.
+	if (bBoneLocationUpdated)
+	{
+		// First step: update bone transform positions from chain links.
+		for (int32 LinkIndex = 0; LinkIndex < NumChainLinks; LinkIndex++)
+		{
+			CCDIKChainLink const & ChainLink = Chain[LinkIndex];
+			OutBoneTransforms[ChainLink.TransformIndex].Transform = ChainLink.Transform;
+
+			// If there are any zero length children, update position of those
+			int32 const NumChildren = ChainLink.ChildZeroLengthTransformIndices.Num();
+			for (int32 ChildIndex = 0; ChildIndex < NumChildren; ChildIndex++)
+			{
+				OutBoneTransforms[ChainLink.ChildZeroLengthTransformIndices[ChildIndex]].Transform = ChainLink.Transform;
+			}
+		}
+
+#if WITH_EDITOR
+		DebugLines.Reset(OutBoneTransforms.Num());
+		DebugLines.AddUninitialized(OutBoneTransforms.Num());
+		for (int32 Index = 0; Index < OutBoneTransforms.Num(); ++Index)
+		{
+			DebugLines[Index] = OutBoneTransforms[Index].Transform.GetLocation();
+		}
+#endif // WITH_EDITOR
+
+	}
+}
+
+bool FAnimNode_CCDIK::UpdateChainLink(TArray<CCDIKChainLink>& Chain, int32 LinkIndex, const FVector& TargetPos) const
+{
+	int32 const TipBoneLinkIndex = Chain.Num() - 1;
+
+	ensure(Chain.IsValidIndex(TipBoneLinkIndex));
+	CCDIKChainLink& CurrentLink = Chain[LinkIndex];
+
+	// update new tip pos
+	FVector TipPos = Chain[TipBoneLinkIndex].Transform.GetLocation();
+
+	FTransform& CurrentLinkTransform = CurrentLink.Transform;
+	FVector ToEnd = TipPos - CurrentLinkTransform.GetLocation();
+	FVector ToTarget = TargetPos - CurrentLinkTransform.GetLocation();
+
+	ToEnd.Normalize();
+	ToTarget.Normalize();
+
+	float RotationLimitPerJointInRadian = FMath::DegreesToRadians(RotationLimitPerJoints[LinkIndex]);
+	float Angle = FMath::ClampAngle(FMath::Acos(FVector::DotProduct(ToEnd, ToTarget)), -RotationLimitPerJointInRadian, RotationLimitPerJointInRadian);
+	bool bCanRotate = (FMath::Abs(Angle) > KINDA_SMALL_NUMBER) && (!bEnableRotationLimit || RotationLimitPerJointInRadian > CurrentLink.CurrentAngleDelta);
+	if (bCanRotate)
+	{
+		// check rotation limit first, if fails, just abort
+		if (bEnableRotationLimit)
+		{
+			if (RotationLimitPerJointInRadian < CurrentLink.CurrentAngleDelta + Angle)
+			{
+				Angle = RotationLimitPerJointInRadian - CurrentLink.CurrentAngleDelta;
+				if (Angle <= KINDA_SMALL_NUMBER)
+				{
+					return false;
+				}
+			}
+
+			CurrentLink.CurrentAngleDelta += Angle;
+		}
+
+		// continue with rotating toward to target
+		FVector RotationAxis = FVector::CrossProduct(ToEnd, ToTarget);
+		if (RotationAxis.SizeSquared() > 0.f)
+		{
+			RotationAxis.Normalize();
+			// Delta Rotation is the rotation to target
+			FQuat DeltaRotation(RotationAxis, Angle);
+
+			FQuat NewRotation = DeltaRotation * CurrentLinkTransform.GetRotation();
+			NewRotation.Normalize();
+			CurrentLinkTransform.SetRotation(NewRotation);
+
+			// if I have parent, make sure to refresh local transform since my current transform has changed
+			if (LinkIndex > 0)
+			{
+				CCDIKChainLink const & Parent = Chain[LinkIndex - 1];
+				CurrentLink.LocalTransform = CurrentLinkTransform.GetRelativeTransform(Parent.Transform);
+				CurrentLink.LocalTransform.NormalizeRotation();
+			}
+
+			// now update all my children to have proper transform
+			FTransform CurrentParentTransform = CurrentLinkTransform;
+
+			// now update all chain
+			for (int32 ChildLinkIndex = LinkIndex + 1; ChildLinkIndex <= TipBoneLinkIndex; ++ChildLinkIndex)
+			{
+				CCDIKChainLink& ChildIterLink = Chain[ChildLinkIndex];
+				FCompactPoseBoneIndex ChildBoneIndex = ChildIterLink.BoneIndex;
+				const FTransform LocalTransform = ChildIterLink.LocalTransform;
+				ChildIterLink.Transform = LocalTransform * CurrentParentTransform;
+				ChildIterLink.Transform.NormalizeRotation();
+				CurrentParentTransform = ChildIterLink.Transform;
+			}
+
+			return true;
+		}
+	}
+
+	return false;
+}
+
+bool FAnimNode_CCDIK::IsValidToEvaluate(const USkeleton* Skeleton, const FBoneContainer& RequiredBones)
+{
+	if (EffectorLocationSpace == BCS_ParentBoneSpace || EffectorLocationSpace == BCS_BoneSpace)
+	{
+		if (!EffectorTarget.IsValidToEvaluate(RequiredBones))
+		{
+			return false;
+		}
+	}
+
+	// Allow evaluation if all parameters are initialized and TipBone is child of RootBone
+	return
+		(
+		TipBone.IsValidToEvaluate(RequiredBones)
+		&& RootBone.IsValidToEvaluate(RequiredBones)
+		&& Precision > 0
+		&& RequiredBones.BoneIsChildOf(TipBone.BoneIndex, RootBone.BoneIndex)
+		);
+}
+
+#if WITH_EDITOR
+void FAnimNode_CCDIK::ResizeRotationLimitPerJoints(int32 NewSize)
+{
+	if (NewSize == 0)
+	{
+		RotationLimitPerJoints.Reset();
+	}
+	else if (RotationLimitPerJoints.Num() != NewSize)
+	{
+		int32 StartIndex = RotationLimitPerJoints.Num();
+		RotationLimitPerJoints.SetNum(NewSize);
+		for (int32 Index = StartIndex; Index < RotationLimitPerJoints.Num(); ++Index)
+		{
+			RotationLimitPerJoints[Index] = 30.f;
+		}
+	}
+}
+#endif 
+
+void FAnimNode_CCDIK::InitializeBoneReferences(const FBoneContainer& RequiredBones)
+{
+	TipBone.Initialize(RequiredBones);
+	RootBone.Initialize(RequiredBones);
+	EffectorTarget.InitializeBoneReferences(RequiredBones);
+}
+
+void FAnimNode_CCDIK::GatherDebugData(FNodeDebugData& DebugData)
+{
+	FString DebugLine = DebugData.GetNodeName(this);
+
+	DebugData.AddDebugItem(DebugLine);
+	ComponentPose.GatherDebugData(DebugData);
+}
+
+PRAGMA_ENABLE_OPTIMIZATION
\ No newline at end of file
diff --git a/Engine/Source/Runtime/AnimGraphRuntime/Public/BoneControllers/AnimNode_CCDIK.h b/Engine/Source/Runtime/AnimGraphRuntime/Public/BoneControllers/AnimNode_CCDIK.h
new file mode 100644
--- /dev/null
+++ b/Engine/Source/Runtime/AnimGraphRuntime/Public/BoneControllers/AnimNode_CCDIK.h
@@ -0,0 +1,129 @@
+// Copyright 1998-2016 Epic Games, Inc. All Rights Reserved.
+
+#pragma once
+
+#include "AnimNode_SkeletalControlBase.h"
+#include "AnimNode_CCDIK.generated.h"
+
+/**
+*	Controller which implements the CCDIK IK approximation algorithm
+*/
+
+/** Transient structure for CCDIK node evaluation */
+struct CCDIKChainLink
+{
+public:
+	/** Transform of bone in component space. */
+	FTransform Transform;
+
+	/** Transform of bone in local space. This is mutable as their component space changes or parents*/
+	FTransform LocalTransform;
+
+	/** Bone Index in SkeletalMesh */
+	FCompactPoseBoneIndex BoneIndex;
+
+	/** Transform Index that this control will output */
+	int32 TransformIndex;
+
+	/** Child bones which are overlapping this bone. 
+	 * They have a zero length distance, so they will inherit this bone's transformation. */
+	TArray<int32> ChildZeroLengthTransformIndices;
+
+	float CurrentAngleDelta;
+
+	CCDIKChainLink()
+		: BoneIndex(INDEX_NONE)
+		, TransformIndex(INDEX_NONE)
+		, CurrentAngleDelta(0.f)
+	{
+	}
+
+	CCDIKChainLink(const FTransform& InTransform, const FTransform& InLocalTransform,  const FCompactPoseBoneIndex& InBoneIndex, const int32& InTransformIndex)
+		: Transform(InTransform)
+		, LocalTransform(InLocalTransform)
+		, BoneIndex(InBoneIndex)
+		, TransformIndex(InTransformIndex)
+		, CurrentAngleDelta(0.f)
+	{
+	}
+};
+
+USTRUCT()
+struct ANIMGRAPHRUNTIME_API FAnimNode_CCDIK : public FAnimNode_SkeletalControlBase
+{
+	GENERATED_USTRUCT_BODY()
+
+	/** Coordinates for target location of tip bone - if EffectorLocationSpace is bone, this is the offset from Target Bone to use as target location*/
+	UPROPERTY(EditAnywhere, Category = Effector, meta = (PinShownByDefault))
+	FVector EffectorLocation;
+
+	/** Reference frame of Effector Transform. */
+	UPROPERTY(EditAnywhere, Category = Effector)
+	TEnumAsByte<enum EBoneControlSpace> EffectorLocationSpace;
+
+	/** If EffectorTransformSpace is a bone, this is the bone to use. **/
+	UPROPERTY(EditAnywhere, Category = Effector)
+	FBoneSocketTarget EffectorTarget;
+
+	/** Name of tip bone */
+	UPROPERTY(EditAnywhere, Category = Solver)
+	FBoneReference TipBone;
+
+	/** Name of the root bone*/
+	UPROPERTY(EditAnywhere, Category = Solver)
+	FBoneReference RootBone;
+
+	/** Tolerance for final tip location delta from EffectorLocation*/
+	UPROPERTY(EditAnywhere, Category = Solver)
+	float Precision;
+
+	/** Maximum number of iterations allowed, to control performance. */
+	UPROPERTY(EditAnywhere, Category = Solver)
+	int32 MaxIterations;
+
+	/** Toggle drawing of axes to debug joint rotation*/
+	UPROPERTY(EditAnywhere, Category = Solver)
+	bool bStartFromTail;
+
+	/** Tolerance for final tip location delta from EffectorLocation*/
+	UPROPERTY(EditAnywhere, Category = Solver)
+	bool bEnableRotationLimit;
+
+private:
+	/** symmetry rotation limit per joint. Index 0 matches with root bone and last index matches with tip bone. */
+	UPROPERTY(EditAnywhere, EditFixedSize, Category = Solver)
+	TArray<float> RotationLimitPerJoints;
+
+public:
+	FAnimNode_CCDIK();
+
+	// FAnimNode_Base interface
+	virtual void GatherDebugData(FNodeDebugData& DebugData) override;
+	// End of FAnimNode_Base interface
+
+	// FAnimNode_SkeletalControlBase interface
+	virtual void EvaluateSkeletalControl_AnyThread(FComponentSpacePoseContext& Output, TArray<FBoneTransform>& OutBoneTransforms) override;
+	virtual bool IsValidToEvaluate(const USkeleton* Skeleton, const FBoneContainer& RequiredBones) override;
+	// End of FAnimNode_SkeletalControlBase interface
+
+private:
+	// FAnimNode_SkeletalControlBase interface
+	virtual void InitializeBoneReferences(const FBoneContainer& RequiredBones) override;
+	// End of FAnimNode_SkeletalControlBase interface
+
+	// Convenience function to get current (pre-translation iteration) component space location of bone by bone index
+	FVector GetCurrentLocation(FCSPose<FCompactPose>& MeshBases, const FCompactPoseBoneIndex& BoneIndex);
+
+	static FTransform GetTargetTransform(const FTransform& InComponentTransform, FCSPose<FCompactPose>& MeshBases, FBoneSocketTarget& InTarget, EBoneControlSpace Space, const FVector& InOffset);
+
+	// return true if updated
+	bool UpdateChainLink(TArray<CCDIKChainLink>& Chain, int32 LinkIndex, const FVector& TargetPos) const;
+public:
+#if WITH_EDITOR
+#if !(UE_BUILD_SHIPPING || UE_BUILD_TEST)
+	TArray<FVector> DebugLines;
+#endif // #if !(UE_BUILD_SHIPPING || UE_BUILD_TEST)
+	// resize rotation limit array based on set up
+	void ResizeRotationLimitPerJoints(int32 NewSize);
+#endif
+};
\ No newline at end of file
