Fixup fallout from Main merge

diff --git a/Engine/Source/Runtime/AnimGraphRuntime/Public/AnimNodes/AnimNode_StateResult.cpp b/Engine/Source/Runtime/AnimGraphRuntime/Public/AnimNodes/AnimNode_StateResult.cpp
--- a/Engine/Source/Runtime/AnimGraphRuntime/Public/AnimNodes/AnimNode_StateResult.cpp
+++ b/Engine/Source/Runtime/AnimGraphRuntime/Public/AnimNodes/AnimNode_StateResult.cpp
@@ -2,12 +2,12 @@
 
 #include "AnimNode_StateResult.h"
 
-bool FAnimNode_StateResult::SerializeFromMismatchedTag(const FPropertyTag& Tag, FArchive& Ar)
+bool FAnimNode_StateResult::SerializeFromMismatchedTag(const FPropertyTag& Tag, FStructuredArchive::FSlot Slot)
 {
 	if(Tag.Type == NAME_StructProperty && Tag.StructName == FAnimNode_Root::StaticStruct()->GetFName())
 	{
 		FAnimNode_Root OldValue;
-		FAnimNode_Root::StaticStruct()->SerializeItem(Ar, &OldValue, nullptr);
+		FAnimNode_Root::StaticStruct()->SerializeItem(Slot, &OldValue, nullptr);
 		*static_cast<FAnimNode_Root*>(this) = OldValue;
 
 		return true;
diff --git a/Engine/Source/Runtime/AnimGraphRuntime/Public/AnimNodes/AnimNode_StateResult.h b/Engine/Source/Runtime/AnimGraphRuntime/Public/AnimNodes/AnimNode_StateResult.h
--- a/Engine/Source/Runtime/AnimGraphRuntime/Public/AnimNodes/AnimNode_StateResult.h
+++ b/Engine/Source/Runtime/AnimGraphRuntime/Public/AnimNodes/AnimNode_StateResult.h
@@ -15,7 +15,7 @@ struct ANIMGRAPHRUNTIME_API FAnimNode_StateResult : public FAnimNode_Root
 	GENERATED_USTRUCT_BODY()
 
 	/** Used to upgrade old FAnimNode_Roots to FAnimNode_StateResult */
-	bool SerializeFromMismatchedTag(const FPropertyTag& Tag, FArchive& Ar);
+	bool SerializeFromMismatchedTag(const FPropertyTag& Tag, FStructuredArchive::FSlot Slot);
 };
 
 template<>
