Fix incorrect interpolations when using URO and CopyPoseFromMesh

This ensures that CopyPoseFromMesh uses the correct target pose to inteprolate to, rather than the already-interpolated output pose.

diff --git a/Engine/Source/Runtime/AnimGraphRuntime/Private/AnimNodes/AnimNode_CopyPoseFromMesh.cpp b/Engine/Source/Runtime/AnimGraphRuntime/Private/AnimNodes/AnimNode_CopyPoseFromMesh.cpp
--- a/Engine/Source/Runtime/AnimGraphRuntime/Private/AnimNodes/AnimNode_CopyPoseFromMesh.cpp
+++ b/Engine/Source/Runtime/AnimGraphRuntime/Private/AnimNodes/AnimNode_CopyPoseFromMesh.cpp
@@ -85,7 +85,9 @@ void FAnimNode_CopyPoseFromMesh::Evaluate_AnyThread(FPoseContext& Output)
 	if (CurrentMeshComponent && CurrentMeshComponent->SkeletalMesh && CurrentMeshComponent->IsRegistered())
 	{
 		const FBoneContainer& RequiredBones = OutPose.GetBoneContainer();
-		const TArray<FTransform>& SourcMeshTransformArray = CurrentMeshComponent->GetComponentSpaceTransforms();
+		const bool bUROInSync = CurrentMeshComponent->AnimUpdateRateParams != nullptr && CurrentMeshComponent->AnimUpdateRateParams == Output.AnimInstanceProxy->GetSkelMeshComponent()->AnimUpdateRateParams;
+		const bool bArraySizesMatch = CurrentMeshComponent->CachedComponentSpaceTransforms.Num() == CurrentMeshComponent->GetComponentSpaceTransforms().Num();
+		const TArray<FTransform>& SourceMeshTransformArray =  bUROInSync && bArraySizesMatch ? CurrentMeshComponent->CachedComponentSpaceTransforms : CurrentMeshComponent->GetComponentSpaceTransforms();
 		for(FCompactPoseBoneIndex PoseBoneIndex : OutPose.ForEachBoneIndex())
 		{
 			const int32& SkeletonBoneIndex = RequiredBones.GetSkeletonIndex(PoseBoneIndex);
@@ -99,13 +101,13 @@ void FAnimNode_CopyPoseFromMesh::Evaluate_AnyThread(FPoseContext& Output)
 				// only apply if I also have parent, otherwise, it should apply the space bases
 				if (ParentIndex != INDEX_NONE && MyParentIndex != INDEX_NONE)
 				{
-					const FTransform& ParentTransform = SourcMeshTransformArray[ParentIndex];
-					const FTransform& ChildTransform = SourcMeshTransformArray[SourceBoneIndex];
+					const FTransform& ParentTransform = SourceMeshTransformArray[ParentIndex];
+					const FTransform& ChildTransform = SourceMeshTransformArray[SourceBoneIndex];
 					OutPose[PoseBoneIndex] = ChildTransform.GetRelativeTransform(ParentTransform);
 				}
 				else
 				{
-					OutPose[PoseBoneIndex] = SourcMeshTransformArray[SourceBoneIndex];
+					OutPose[PoseBoneIndex] = SourceMeshTransformArray[SourceBoneIndex];
 				}
 			}
 		}
