Sequencer: Revert update on bind because it will be updated in the sequencer template if the pose hasn't been ticked yet. See 4229000

diff --git a/Engine/Source/Runtime/AnimGraphRuntime/Public/AnimCustomInstance.h b/Engine/Source/Runtime/AnimGraphRuntime/Public/AnimCustomInstance.h
--- a/Engine/Source/Runtime/AnimGraphRuntime/Public/AnimCustomInstance.h
+++ b/Engine/Source/Runtime/AnimGraphRuntime/Public/AnimCustomInstance.h
@@ -50,7 +50,6 @@ class ANIMGRAPHRUNTIME_API UAnimCustomInstance : public UAnimInstance
 				InstanceClassType* SequencerInstance = NewObject<InstanceClassType>(InSkeletalMeshComponent, InstanceClassType::StaticClass());
 				InSkeletalMeshComponent->AnimScriptInstance = SequencerInstance;
 				InSkeletalMeshComponent->AnimScriptInstance->InitializeAnimation();
-				SequencerInstance->bNeedsUpdate = true;
 				return SequencerInstance;
 			}
 			else
