Speculated fix w.r.t. copy pose node
- where it seems the reference contains invalid skeleton
- refSkeleton access seems violated
- copied Jaren's fix w.r.t. evaluation crash due to setting different AnimBP in the middle of evaluation

diff --git a/Engine/Source/Runtime/AnimGraphRuntime/Private/AnimNodes/AnimNode_CopyPoseFromMesh.cpp b/Engine/Source/Runtime/AnimGraphRuntime/Private/AnimNodes/AnimNode_CopyPoseFromMesh.cpp
--- a/Engine/Source/Runtime/AnimGraphRuntime/Private/AnimNodes/AnimNode_CopyPoseFromMesh.cpp
+++ b/Engine/Source/Runtime/AnimGraphRuntime/Private/AnimNodes/AnimNode_CopyPoseFromMesh.cpp
@@ -113,7 +113,7 @@ void FAnimNode_CopyPoseFromMesh::PreUpdate(const UAnimInstance* InAnimInstance)
 			SourceMeshTransformArray.Append((bUROInSync || bUsingExternalInterpolation) && bArraySizesMatch ? CachedComponentSpaceTransforms : CurrentMeshComponent->GetComponentSpaceTransforms());
 
 			// Ref skeleton is need for parent index lookups later, so store it now
-			RefSkeleton = &CurrentMeshComponent->SkeletalMesh->RefSkeleton;
+			CurrentlyUsedMesh = CurrentMeshComponent->SkeletalMesh;
 
 			if(bCopyCurves)
 			{
@@ -130,6 +130,10 @@ void FAnimNode_CopyPoseFromMesh::PreUpdate(const UAnimInstance* InAnimInstance)
 				}
 			}
 		}
+		else
+		{
+			CurrentlyUsedMesh.Reset();
+		}
 	}
 }
 
@@ -144,8 +148,8 @@ void FAnimNode_CopyPoseFromMesh::Evaluate_AnyThread(FPoseContext& Output)
 {
 	FCompactPose& OutPose = Output.Pose;
 	OutPose.ResetToRefPose();
-
-	if(SourceMeshTransformArray.Num() > 0 && RefSkeleton != nullptr)
+	USkeletalMesh* CurrentMesh = CurrentlyUsedMesh.IsValid() ? CurrentlyUsedMesh.Get() : nullptr;
+	if(SourceMeshTransformArray.Num() > 0 && CurrentMesh)
 	{
 		const FBoneContainer& RequiredBones = OutPose.GetBoneContainer();
 
@@ -157,7 +161,7 @@ void FAnimNode_CopyPoseFromMesh::Evaluate_AnyThread(FPoseContext& Output)
 			if(Value && SourceMeshTransformArray.IsValidIndex(*Value))
 			{
 				const int32 SourceBoneIndex = *Value;
-				const int32 ParentIndex = RefSkeleton->GetParentIndex(SourceBoneIndex);
+				const int32 ParentIndex = CurrentMesh->RefSkeleton.GetParentIndex(SourceBoneIndex);
 				const FCompactPoseBoneIndex MyParentIndex = RequiredBones.GetParentBoneIndex(PoseBoneIndex);
 				// only apply if I also have parent, otherwise, it should apply the space bases
 				if (SourceMeshTransformArray.IsValidIndex(ParentIndex) && MyParentIndex != INDEX_NONE)
@@ -198,7 +202,7 @@ void FAnimNode_CopyPoseFromMesh::GatherDebugData(FNodeDebugData& DebugData)
 
 void FAnimNode_CopyPoseFromMesh::ReinitializeMeshComponent(USkeletalMeshComponent* NewSourceMeshComponent, USkeletalMeshComponent* TargetMeshComponent)
 {
-	CurrentlyUsedSourceMeshComponent = NewSourceMeshComponent;
+	CurrentlyUsedSourceMeshComponent.Reset();
 	// reset source mesh
 	CurrentlyUsedSourceMesh.Reset();
 	BoneMapToSource.Reset();
@@ -208,45 +212,51 @@ void FAnimNode_CopyPoseFromMesh::ReinitializeMeshComponent(USkeletalMeshComponen
 	{
 		USkeletalMesh* SourceSkelMesh = NewSourceMeshComponent->SkeletalMesh;
 		USkeletalMesh* TargetSkelMesh = TargetMeshComponent->SkeletalMesh;
-		CurrentlyUsedSourceMesh = SourceSkelMesh;
-
-		if (SourceSkelMesh == TargetSkelMesh)
+		
+		if (SourceSkelMesh && !SourceSkelMesh->IsPendingKill() && !SourceSkelMesh->HasAnyFlags(RF_NeedPostLoad) &&
+			TargetSkelMesh && !TargetSkelMesh->IsPendingKill() && !TargetSkelMesh->HasAnyFlags(RF_NeedPostLoad))
 		{
-			for(int32 ComponentSpaceBoneId = 0; ComponentSpaceBoneId < SourceSkelMesh->RefSkeleton.GetNum(); ++ComponentSpaceBoneId)
+			CurrentlyUsedSourceMeshComponent = NewSourceMeshComponent;
+			CurrentlyUsedSourceMesh = SourceSkelMesh;
+
+			if (SourceSkelMesh == TargetSkelMesh)
 			{
-				BoneMapToSource.Add(ComponentSpaceBoneId, ComponentSpaceBoneId);
+				for(int32 ComponentSpaceBoneId = 0; ComponentSpaceBoneId < SourceSkelMesh->RefSkeleton.GetNum(); ++ComponentSpaceBoneId)
+				{
+					BoneMapToSource.Add(ComponentSpaceBoneId, ComponentSpaceBoneId);
+				}
 			}
-		}
-		else
-		{
-			for (int32 ComponentSpaceBoneId=0; ComponentSpaceBoneId<TargetSkelMesh->RefSkeleton.GetNum(); ++ComponentSpaceBoneId)
+			else
 			{
-				FName BoneName = TargetSkelMesh->RefSkeleton.GetBoneName(ComponentSpaceBoneId);
-				BoneMapToSource.Add(ComponentSpaceBoneId, SourceSkelMesh->RefSkeleton.FindBoneIndex(BoneName));
+				for (int32 ComponentSpaceBoneId=0; ComponentSpaceBoneId<TargetSkelMesh->RefSkeleton.GetNum(); ++ComponentSpaceBoneId)
+				{
+					FName BoneName = TargetSkelMesh->RefSkeleton.GetBoneName(ComponentSpaceBoneId);
+					BoneMapToSource.Add(ComponentSpaceBoneId, SourceSkelMesh->RefSkeleton.FindBoneIndex(BoneName));
+				}
 			}
-		}
 
-		if (bCopyCurves)
-		{
-			USkeleton* SourceSkeleton = SourceSkelMesh->Skeleton;
-			USkeleton* TargetSkeleton = TargetSkelMesh->Skeleton;
+			if (bCopyCurves)
+			{
+				USkeleton* SourceSkeleton = SourceSkelMesh->Skeleton;
+				USkeleton* TargetSkeleton = TargetSkelMesh->Skeleton;
 
-			// you shouldn't be here if this happened
-			check(SourceSkeleton && TargetSkeleton);
+				// you shouldn't be here if this happened
+				check(SourceSkeleton && TargetSkeleton);
 
-			const FSmartNameMapping* SourceContainer = SourceSkeleton->GetSmartNameContainer(USkeleton::AnimCurveMappingName);
-			const FSmartNameMapping* TargetContainer = TargetSkeleton->GetSmartNameContainer(USkeleton::AnimCurveMappingName);
+				const FSmartNameMapping* SourceContainer = SourceSkeleton->GetSmartNameContainer(USkeleton::AnimCurveMappingName);
+				const FSmartNameMapping* TargetContainer = TargetSkeleton->GetSmartNameContainer(USkeleton::AnimCurveMappingName);
 
-			TArray<FName> SourceCurveNames;
-			SourceContainer->FillNameArray(SourceCurveNames);
-			for (int32 Index = 0; Index < SourceCurveNames.Num(); ++Index)
-			{
-				SmartName::UID_Type UID = TargetContainer->FindUID(SourceCurveNames[Index]);
-				if (UID != SmartName::MaxUID)
+				TArray<FName> SourceCurveNames;
+				SourceContainer->FillNameArray(SourceCurveNames);
+				for (int32 Index = 0; Index < SourceCurveNames.Num(); ++Index)
 				{
-					// has a valid UID, add to the list
-					SmartName::UID_Type& Value = CurveNameToUIDMap.Add(SourceCurveNames[Index]);
-					Value = UID;
+					SmartName::UID_Type UID = TargetContainer->FindUID(SourceCurveNames[Index]);
+					if (UID != SmartName::MaxUID)
+					{
+						// has a valid UID, add to the list
+						SmartName::UID_Type& Value = CurveNameToUIDMap.Add(SourceCurveNames[Index]);
+						Value = UID;
+					}
 				}
 			}
 		}
diff --git a/Engine/Source/Runtime/AnimGraphRuntime/Public/AnimNodes/AnimNode_CopyPoseFromMesh.h b/Engine/Source/Runtime/AnimGraphRuntime/Public/AnimNodes/AnimNode_CopyPoseFromMesh.h
--- a/Engine/Source/Runtime/AnimGraphRuntime/Public/AnimNodes/AnimNode_CopyPoseFromMesh.h
+++ b/Engine/Source/Runtime/AnimGraphRuntime/Public/AnimNodes/AnimNode_CopyPoseFromMesh.h
@@ -20,7 +20,7 @@ struct ANIMGRAPHRUNTIME_API FAnimNode_CopyPoseFromMesh : public FAnimNode_Base
 	GENERATED_USTRUCT_BODY()
 
 	/*  This is used by default if it's valid */
-	UPROPERTY(EditAnywhere, BlueprintReadWrite, Category=Copy, meta=(PinShownByDefault))
+	UPROPERTY(BlueprintReadWrite, transient, Category=Copy, meta=(PinShownByDefault))
 	TWeakObjectPtr<USkeletalMeshComponent> SourceMeshComponent;
 
 	/* If SourceMeshComponent is not valid, and if this is true, it will look for attahced parent as a source */
@@ -46,6 +46,7 @@ struct ANIMGRAPHRUNTIME_API FAnimNode_CopyPoseFromMesh : public FAnimNode_Base
 private:
 	TWeakObjectPtr<USkeletalMeshComponent>	CurrentlyUsedSourceMeshComponent;
 	TWeakObjectPtr<USkeletalMesh>			CurrentlyUsedSourceMesh;
+	TWeakObjectPtr<USkeletalMesh>			CurrentlyUsedMesh;
 	// cache of target space bases to source space bases
 	TMap<int32, int32> BoneMapToSource;
 	TMap<FName, SmartName::UID_Type> CurveNameToUIDMap;
@@ -56,9 +57,6 @@ private:
 	// Cached curves, copied on the game thread
 	TMap<FName, float> SourceCurveList;
 
-	// Reference skeleton used for parent bone lookups
-	FReferenceSkeleton* RefSkeleton;
-
 	// reinitialize mesh component 
 	void ReinitializeMeshComponent(USkeletalMeshComponent* NewSkeletalMeshComponent, USkeletalMeshComponent* TargetMeshComponent);
 	void RefreshMeshComponent(USkeletalMeshComponent* TargetMeshComponent);
diff --git a/Engine/Source/Runtime/Engine/Private/Components/SkeletalMeshComponent.cpp b/Engine/Source/Runtime/Engine/Private/Components/SkeletalMeshComponent.cpp
--- a/Engine/Source/Runtime/Engine/Private/Components/SkeletalMeshComponent.cpp
+++ b/Engine/Source/Runtime/Engine/Private/Components/SkeletalMeshComponent.cpp
@@ -793,6 +793,13 @@ void USkeletalMeshComponent::ClearAnimScriptInstance()
 {
 	if (AnimScriptInstance)
 	{
+		// We may be doing parallel evaluation on the current anim instance
+		// Calling this here with true will block this init till that thread completes
+		// and it is safe to continue
+		const bool bBlockOnTask = true; // wait on evaluation task so it is safe to swap the buffers
+		const bool bPerformPostAnimEvaluation = true; // Do PostEvaluation so we make sure to swap the buffers back. 
+		HandleExistingParallelEvaluationTask(bBlockOnTask, bPerformPostAnimEvaluation);
+
 		AnimScriptInstance->EndNotifyStates();
 	}
 	AnimScriptInstance = nullptr;
