Fix clamp linear translation limit causing jittering in some assets.

diff --git a/Engine/Source/Runtime/AnimGraphRuntime/Private/BoneControllers/AnimNode_RigidBody.cpp b/Engine/Source/Runtime/AnimGraphRuntime/Private/BoneControllers/AnimNode_RigidBody.cpp
--- a/Engine/Source/Runtime/AnimGraphRuntime/Private/BoneControllers/AnimNode_RigidBody.cpp
+++ b/Engine/Source/Runtime/AnimGraphRuntime/Private/BoneControllers/AnimNode_RigidBody.cpp
@@ -197,6 +197,7 @@ void FAnimNode_RigidBody::InitializeNewBodyTransformsDuringSimulation(FComponent
 
 				const FTransform WSBodyTM = BodyRelativeTransform * Bodies[OutputData.ParentBodyIndex]->GetWorldTransform();
 				Bodies[BodyIndex]->SetWorldTransform(WSBodyTM);
+				BodyAnimData[BodyIndex].RefPoseLength = BodyRelativeTransform.GetLocation().Size();
 			}
 			// If we don't have a parent body, then we can just grab the incoming pose in component space.
 			else
@@ -323,6 +324,10 @@ void FAnimNode_RigidBody::EvaluateSkeletalControl_AnyThread(FComponentSpacePoseC
 
 						BodyTM = ConvertCSTransformToSimSpace(SimulationSpace, ComponentSpaceTM, CompWorldSpaceTM, BaseBoneTM);
 						Bodies[BodyIndex]->SetWorldTransform(BodyTM);
+						if (OutputData.ParentBodyIndex != INDEX_NONE)
+						{
+							BodyAnimData[BodyIndex].RefPoseLength = BodyTM.GetRelativeTransform(Bodies[OutputData.ParentBodyIndex]->GetWorldTransform()).GetLocation().Size();
+						}
 					}
 				}
 				break;
@@ -338,6 +343,10 @@ void FAnimNode_RigidBody::EvaluateSkeletalControl_AnyThread(FComponentSpacePoseC
 						const FTransform& ComponentSpaceTM = Output.Pose.GetComponentSpaceTransform(OutputData.CompactPoseBoneIndex);
 						const FTransform BodyTM = ConvertCSTransformToSimSpace(SimulationSpace, ComponentSpaceTM, CompWorldSpaceTM, BaseBoneTM);
 						Bodies[BodyIndex]->SetWorldTransform(BodyTM);
+						if (OutputData.ParentBodyIndex != INDEX_NONE)
+						{
+							BodyAnimData[BodyIndex].RefPoseLength = BodyTM.GetRelativeTransform(Bodies[OutputData.ParentBodyIndex]->GetWorldTransform()).GetLocation().Size();
+						}
 					}
 				}
 				break;
