Layer node initialization fix

diff --git a/Engine/Source/Editor/AnimGraph/Classes/AnimGraphNode_LayeredBoneBlend.h b/Engine/Source/Editor/AnimGraph/Classes/AnimGraphNode_LayeredBoneBlend.h
--- a/Engine/Source/Editor/AnimGraph/Classes/AnimGraphNode_LayeredBoneBlend.h
+++ b/Engine/Source/Editor/AnimGraph/Classes/AnimGraphNode_LayeredBoneBlend.h
@@ -37,5 +37,7 @@ class UAnimGraphNode_LayeredBoneBlend : public UAnimGraphNode_BlendListBase
 
 	// UAnimGraphNode_Base interface
 	virtual FString GetNodeCategory() const override;
+	// Gives each visual node a chance to validate that they are still valid in the context of the compiled class, giving a last shot at error or warning generation after primary compilation is finished
+	virtual void ValidateAnimNodeDuringCompilation(class USkeleton* ForSkeleton, class FCompilerResultsLog& MessageLog) override;
 	// End of UAnimGraphNode_Base interface
 };
diff --git a/Engine/Source/Editor/AnimGraph/Private/AnimGraphNode_LayeredBoneBlend.cpp b/Engine/Source/Editor/AnimGraph/Private/AnimGraphNode_LayeredBoneBlend.cpp
--- a/Engine/Source/Editor/AnimGraph/Private/AnimGraphNode_LayeredBoneBlend.cpp
+++ b/Engine/Source/Editor/AnimGraph/Private/AnimGraphNode_LayeredBoneBlend.cpp
@@ -97,4 +97,15 @@ void UAnimGraphNode_LayeredBoneBlend::Serialize(FArchive& Ar)
 	Super::Serialize(Ar);
 	Node.ValidateData();
 }
+
+void UAnimGraphNode_LayeredBoneBlend::ValidateAnimNodeDuringCompilation(class USkeleton* ForSkeleton, class FCompilerResultsLog& MessageLog)
+{
+	UAnimGraphNode_Base::ValidateAnimNodeDuringCompilation(ForSkeleton, MessageLog);
+
+	// ensure to cache the data
+ 	if (Node.IsCacheInvalid(ForSkeleton))
+ 	{
+ 		Node.RebuildCacheData(ForSkeleton);
+ 	}
+}
 #undef LOCTEXT_NAMESPACE
diff --git a/Engine/Source/Runtime/AnimGraphRuntime/Private/AnimNodes/AnimNode_LayeredBoneBlend.cpp b/Engine/Source/Runtime/AnimGraphRuntime/Private/AnimNodes/AnimNode_LayeredBoneBlend.cpp
--- a/Engine/Source/Runtime/AnimGraphRuntime/Private/AnimNodes/AnimNode_LayeredBoneBlend.cpp
+++ b/Engine/Source/Runtime/AnimGraphRuntime/Private/AnimNodes/AnimNode_LayeredBoneBlend.cpp
@@ -41,8 +41,7 @@ void FAnimNode_LayeredBoneBlend::PostCompile(const class USkeleton* InSkeleton)
 
 void FAnimNode_LayeredBoneBlend::RebuildCacheData(const USkeleton* InSkeleton)
 {
-	// make sure it's Finished post load yet, this might cause more things to cache in initialize instead of during cooking
-	if (InSkeleton && !(InSkeleton->GetFlags() & RF_NeedPostLoad))
+	if (InSkeleton)
 	{
 		FAnimationRuntime::CreateMaskWeights(PerBoneBlendWeights, LayerSetup, InSkeleton);
 		SkeletonGuid = InSkeleton->GetGuid();
diff --git a/Engine/Source/Runtime/AnimGraphRuntime/Public/AnimNodes/AnimNode_LayeredBoneBlend.h b/Engine/Source/Runtime/AnimGraphRuntime/Public/AnimNodes/AnimNode_LayeredBoneBlend.h
--- a/Engine/Source/Runtime/AnimGraphRuntime/Public/AnimNodes/AnimNode_LayeredBoneBlend.h
+++ b/Engine/Source/Runtime/AnimGraphRuntime/Public/AnimNodes/AnimNode_LayeredBoneBlend.h
@@ -129,4 +129,6 @@ private:
 	// Rebuild cache data from the skeleton
 	void RebuildCacheData(const USkeleton* InSkeleton);
 	bool IsCacheInvalid(const USkeleton* InSkeleton) const;
+
+	friend class UAnimGraphNode_LayeredBoneBlend;
 };
