Add a physics setting to stop simulating anim dynamics if the timestep is below a threshold.

diff --git a/Engine/Source/Runtime/AnimGraphRuntime/Private/BoneControllers/AnimNode_AnimDynamics.cpp b/Engine/Source/Runtime/AnimGraphRuntime/Private/BoneControllers/AnimNode_AnimDynamics.cpp
--- a/Engine/Source/Runtime/AnimGraphRuntime/Private/BoneControllers/AnimNode_AnimDynamics.cpp
+++ b/Engine/Source/Runtime/AnimGraphRuntime/Private/BoneControllers/AnimNode_AnimDynamics.cpp
@@ -243,7 +243,7 @@ void FAnimNode_AnimDynamics::EvaluateSkeletalControl_AnyThread(FComponentSpacePo
 			}
 		}
 
-		if (bDoUpdate && NextTimeStep > 0.0f)
+		if (bDoUpdate && NextTimeStep > AnimDynamicsMinDeltaTime)
 		{
 			// Calculate gravity direction
 			SimSpaceGravityDirection = TransformWorldVectorToSimSpace(Output, FVector(0.0f, 0.0f, -1.0f));
@@ -702,12 +702,14 @@ void FAnimNode_AnimDynamics::InitPhysics(FComponentSpacePoseContext& Output)
 			// Cache physics settings to avoid accessing UPhysicsSettings continuously
 			if(UPhysicsSettings* Settings = UPhysicsSettings::Get())
 			{
+				AnimDynamicsMinDeltaTime = Settings->AnimDynamicsMinDeltaTime;
 				MaxPhysicsDeltaTime = Settings->MaxPhysicsDeltaTime;
 				MaxSubstepDeltaTime = Settings->MaxSubstepDeltaTime;
 				MaxSubsteps = Settings->MaxSubsteps;
 			}
 			else
 			{
+				AnimDynamicsMinDeltaTime = 0.f;
 				MaxPhysicsDeltaTime = (1.0f/30.0f);
 				MaxSubstepDeltaTime = (1.0f/60.0f);
 				MaxSubsteps = 4;
diff --git a/Engine/Source/Runtime/AnimGraphRuntime/Public/BoneControllers/AnimNode_AnimDynamics.h b/Engine/Source/Runtime/AnimGraphRuntime/Public/BoneControllers/AnimNode_AnimDynamics.h
--- a/Engine/Source/Runtime/AnimGraphRuntime/Public/BoneControllers/AnimNode_AnimDynamics.h
+++ b/Engine/Source/Runtime/AnimGraphRuntime/Public/BoneControllers/AnimNode_AnimDynamics.h
@@ -471,6 +471,7 @@ private:
 	float TimeDebt;
 
 	// Cached physics settings. We cache these on initialise to avoid the cost of accessing UPhysicsSettings a lot each frame
+	float AnimDynamicsMinDeltaTime;
 	float MaxPhysicsDeltaTime;
 	float MaxSubstepDeltaTime;
 	int32 MaxSubsteps;
diff --git a/Engine/Source/Runtime/Engine/Classes/PhysicsEngine/PhysicsSettings.h b/Engine/Source/Runtime/Engine/Classes/PhysicsEngine/PhysicsSettings.h
--- a/Engine/Source/Runtime/Engine/Classes/PhysicsEngine/PhysicsSettings.h
+++ b/Engine/Source/Runtime/Engine/Classes/PhysicsEngine/PhysicsSettings.h
@@ -216,6 +216,10 @@ class ENGINE_API UPhysicsSettings : public UDeveloperSettings
 	UPROPERTY(config, EditAnywhere, Category = Simulation)
 	bool bEnableEnhancedDeterminism;
 
+	/** Min Delta Time below which anim dynamics nodes will not simulate. */
+	UPROPERTY(config, EditAnywhere, meta = (ClampMin = "0.0", UIMin = "0.0", ClampMax = "1.0", UIMax = "1.0"), Category = Framerate)
+	float AnimDynamicsMinDeltaTime;
+
 	/** Max Physics Delta Time to be clamped. */
 	UPROPERTY(config, EditAnywhere, meta=(ClampMin="0.0013", UIMin = "0.0013", ClampMax="1.0", UIMax="1.0"), Category=Framerate)
 	float MaxPhysicsDeltaTime;
diff --git a/Engine/Source/Runtime/Engine/Private/PhysicsEngine/PhysicsSettings.cpp b/Engine/Source/Runtime/Engine/Private/PhysicsEngine/PhysicsSettings.cpp
--- a/Engine/Source/Runtime/Engine/Private/PhysicsEngine/PhysicsSettings.cpp
+++ b/Engine/Source/Runtime/Engine/Private/PhysicsEngine/PhysicsSettings.cpp
@@ -30,6 +30,7 @@ UPhysicsSettings::UPhysicsSettings(const FObjectInitializer& ObjectInitializer)
 	, bSuppressFaceRemapTable(false)
 	, bDisableActiveActors(false)
 	, bEnableEnhancedDeterminism(false)
+	, AnimDynamicsMinDeltaTime(0.f)
 	, MaxPhysicsDeltaTime(1.f / 30.f)
 	, bSubstepping(false)
 	, MaxSubstepDeltaTime(1.f / 60.f)
