Add setting to stop simulating rigidbody nodes below a minimum timestep.

diff --git a/Engine/Source/Runtime/AnimGraphRuntime/Private/BoneControllers/AnimNode_AnimDynamics.cpp b/Engine/Source/Runtime/AnimGraphRuntime/Private/BoneControllers/AnimNode_AnimDynamics.cpp
--- a/Engine/Source/Runtime/AnimGraphRuntime/Private/BoneControllers/AnimNode_AnimDynamics.cpp
+++ b/Engine/Source/Runtime/AnimGraphRuntime/Private/BoneControllers/AnimNode_AnimDynamics.cpp
@@ -243,7 +243,7 @@ void FAnimNode_AnimDynamics::EvaluateSkeletalControl_AnyThread(FComponentSpacePo
 			}
 		}
 
-		if (bDoUpdate && NextTimeStep > AnimDynamicsMinDeltaTime)
+		if (bDoUpdate && NextTimeStep > AnimPhysicsMinDeltaTime)
 		{
 			// Calculate gravity direction
 			SimSpaceGravityDirection = TransformWorldVectorToSimSpace(Output, FVector(0.0f, 0.0f, -1.0f));
@@ -702,14 +702,14 @@ void FAnimNode_AnimDynamics::InitPhysics(FComponentSpacePoseContext& Output)
 			// Cache physics settings to avoid accessing UPhysicsSettings continuously
 			if(UPhysicsSettings* Settings = UPhysicsSettings::Get())
 			{
-				AnimDynamicsMinDeltaTime = Settings->AnimDynamicsMinDeltaTime;
+				AnimPhysicsMinDeltaTime = Settings->AnimPhysicsMinDeltaTime;
 				MaxPhysicsDeltaTime = Settings->MaxPhysicsDeltaTime;
 				MaxSubstepDeltaTime = Settings->MaxSubstepDeltaTime;
 				MaxSubsteps = Settings->MaxSubsteps;
 			}
 			else
 			{
-				AnimDynamicsMinDeltaTime = 0.f;
+				AnimPhysicsMinDeltaTime = 0.f;
 				MaxPhysicsDeltaTime = (1.0f/30.0f);
 				MaxSubstepDeltaTime = (1.0f/60.0f);
 				MaxSubsteps = 4;
diff --git a/Engine/Source/Runtime/AnimGraphRuntime/Private/BoneControllers/AnimNode_RigidBody.cpp b/Engine/Source/Runtime/AnimGraphRuntime/Private/BoneControllers/AnimNode_RigidBody.cpp
--- a/Engine/Source/Runtime/AnimGraphRuntime/Private/BoneControllers/AnimNode_RigidBody.cpp
+++ b/Engine/Source/Runtime/AnimGraphRuntime/Private/BoneControllers/AnimNode_RigidBody.cpp
@@ -10,7 +10,7 @@
 #include "Physics/PhysicsInterfaceCore.h"
 #include "Physics/ImmediatePhysics/ImmediatePhysicsSimulation.h"
 #include "Physics/ImmediatePhysics/ImmediatePhysicsActorHandle.h"
-
+#include "PhysicsEngine/PhysicsSettings.h"
 #include "Logging/MessageLog.h"
 
 //#pragma optimize("", off)
@@ -358,7 +358,7 @@ void FAnimNode_RigidBody::EvaluateSkeletalControl_AnyThread(FComponentSpacePoseC
 			PreviousComponentLinearVelocity = FVector::ZeroVector;
 		}
 		// Only need to tick physics if we didn't reset and we have some time to simulate
-		else if(DeltaSeconds > 0.0f)
+		else if(DeltaSeconds > AnimPhysicsMinDeltaTime)
 		{
 			// Transfer bone velocities previously captured.
 			if (bTransferBoneVelocities && (CapturedBoneVelocityPose.GetPose().GetNumBones() > 0))
@@ -588,6 +588,15 @@ void FAnimNode_RigidBody::InitPhysics(const UAnimInstance* InAnimInstance)
 
 	PreviousTransform = InAnimInstance->GetSkelMeshComponent()->GetComponentToWorld();
 
+	if (UPhysicsSettings* Settings = UPhysicsSettings::Get())
+	{
+		AnimPhysicsMinDeltaTime = Settings->AnimPhysicsMinDeltaTime;
+	}
+	else
+	{
+		AnimPhysicsMinDeltaTime = 0.f;
+	}
+
 	if(UsePhysicsAsset)
 	{
 		delete PhysicsSimulation;
diff --git a/Engine/Source/Runtime/AnimGraphRuntime/Public/BoneControllers/AnimNode_AnimDynamics.h b/Engine/Source/Runtime/AnimGraphRuntime/Public/BoneControllers/AnimNode_AnimDynamics.h
--- a/Engine/Source/Runtime/AnimGraphRuntime/Public/BoneControllers/AnimNode_AnimDynamics.h
+++ b/Engine/Source/Runtime/AnimGraphRuntime/Public/BoneControllers/AnimNode_AnimDynamics.h
@@ -471,7 +471,7 @@ private:
 	float TimeDebt;
 
 	// Cached physics settings. We cache these on initialise to avoid the cost of accessing UPhysicsSettings a lot each frame
-	float AnimDynamicsMinDeltaTime;
+	float AnimPhysicsMinDeltaTime;
 	float MaxPhysicsDeltaTime;
 	float MaxSubstepDeltaTime;
 	int32 MaxSubsteps;
diff --git a/Engine/Source/Runtime/AnimGraphRuntime/Public/BoneControllers/AnimNode_RigidBody.h b/Engine/Source/Runtime/AnimGraphRuntime/Public/BoneControllers/AnimNode_RigidBody.h
--- a/Engine/Source/Runtime/AnimGraphRuntime/Public/BoneControllers/AnimNode_RigidBody.h
+++ b/Engine/Source/Runtime/AnimGraphRuntime/Public/BoneControllers/AnimNode_RigidBody.h
@@ -174,7 +174,7 @@ private:
 private:
 
 	float AccumulatedDeltaTime;
-
+	float AnimPhysicsMinDeltaTime;
 	/** This should only be used for removing the delegate during termination. Do NOT use this for any per frame work */
 	TWeakObjectPtr<USkeletalMeshComponent> SkelMeshCompWeakPtr;
 
diff --git a/Engine/Source/Runtime/Engine/Classes/PhysicsEngine/PhysicsSettings.h b/Engine/Source/Runtime/Engine/Classes/PhysicsEngine/PhysicsSettings.h
--- a/Engine/Source/Runtime/Engine/Classes/PhysicsEngine/PhysicsSettings.h
+++ b/Engine/Source/Runtime/Engine/Classes/PhysicsEngine/PhysicsSettings.h
@@ -216,9 +216,9 @@ class ENGINE_API UPhysicsSettings : public UDeveloperSettings
 	UPROPERTY(config, EditAnywhere, Category = Simulation)
 	bool bEnableEnhancedDeterminism;
 
-	/** Min Delta Time below which anim dynamics nodes will not simulate. */
+	/** Min Delta Time below which anim dynamics and rigidbody nodes will not simulate. */
 	UPROPERTY(config, EditAnywhere, meta = (ClampMin = "0.0", UIMin = "0.0", ClampMax = "1.0", UIMax = "1.0"), Category = Framerate)
-	float AnimDynamicsMinDeltaTime;
+	float AnimPhysicsMinDeltaTime;
 
 	/** Max Physics Delta Time to be clamped. */
 	UPROPERTY(config, EditAnywhere, meta=(ClampMin="0.0013", UIMin = "0.0013", ClampMax="1.0", UIMax="1.0"), Category=Framerate)
diff --git a/Engine/Source/Runtime/Engine/Private/PhysicsEngine/PhysicsSettings.cpp b/Engine/Source/Runtime/Engine/Private/PhysicsEngine/PhysicsSettings.cpp
--- a/Engine/Source/Runtime/Engine/Private/PhysicsEngine/PhysicsSettings.cpp
+++ b/Engine/Source/Runtime/Engine/Private/PhysicsEngine/PhysicsSettings.cpp
@@ -30,7 +30,7 @@ UPhysicsSettings::UPhysicsSettings(const FObjectInitializer& ObjectInitializer)
 	, bSuppressFaceRemapTable(false)
 	, bDisableActiveActors(false)
 	, bEnableEnhancedDeterminism(false)
-	, AnimDynamicsMinDeltaTime(0.f)
+	, AnimPhysicsMinDeltaTime(0.f)
 	, MaxPhysicsDeltaTime(1.f / 30.f)
 	, bSubstepping(false)
 	, MaxSubstepDeltaTime(1.f / 60.f)
