Add more log w.r.t. skeleton crash from copy pose and change to ensure

diff --git a/Engine/Source/Runtime/AnimGraphRuntime/Private/AnimNodes/AnimNode_CopyPoseFromMesh.cpp b/Engine/Source/Runtime/AnimGraphRuntime/Private/AnimNodes/AnimNode_CopyPoseFromMesh.cpp
--- a/Engine/Source/Runtime/AnimGraphRuntime/Private/AnimNodes/AnimNode_CopyPoseFromMesh.cpp
+++ b/Engine/Source/Runtime/AnimGraphRuntime/Private/AnimNodes/AnimNode_CopyPoseFromMesh.cpp
@@ -241,21 +241,23 @@ void FAnimNode_CopyPoseFromMesh::ReinitializeMeshComponent(USkeletalMeshComponen
 				USkeleton* TargetSkeleton = TargetSkelMesh->Skeleton;
 
 				// you shouldn't be here if this happened
-				check(SourceSkeleton && TargetSkeleton);
-
-				const FSmartNameMapping* SourceContainer = SourceSkeleton->GetSmartNameContainer(USkeleton::AnimCurveMappingName);
-				const FSmartNameMapping* TargetContainer = TargetSkeleton->GetSmartNameContainer(USkeleton::AnimCurveMappingName);
-
-				TArray<FName> SourceCurveNames;
-				SourceContainer->FillNameArray(SourceCurveNames);
-				for (int32 Index = 0; Index < SourceCurveNames.Num(); ++Index)
+				if (ensureMsgf(SourceSkeleton, TEXT("Invalid null source skeleton : %s"), *GetNameSafe(SourceSkelMesh))
+					&& ensureMsgf(TargetSkeleton, TEXT("Invalid null target skeleton : %s"), *GetNameSafe(TargetSkelMesh)))
 				{
-					SmartName::UID_Type UID = TargetContainer->FindUID(SourceCurveNames[Index]);
-					if (UID != SmartName::MaxUID)
+					const FSmartNameMapping* SourceContainer = SourceSkeleton->GetSmartNameContainer(USkeleton::AnimCurveMappingName);
+					const FSmartNameMapping* TargetContainer = TargetSkeleton->GetSmartNameContainer(USkeleton::AnimCurveMappingName);
+
+					TArray<FName> SourceCurveNames;
+					SourceContainer->FillNameArray(SourceCurveNames);
+					for (int32 Index = 0; Index < SourceCurveNames.Num(); ++Index)
 					{
-						// has a valid UID, add to the list
-						SmartName::UID_Type& Value = CurveNameToUIDMap.Add(SourceCurveNames[Index]);
-						Value = UID;
+						SmartName::UID_Type UID = TargetContainer->FindUID(SourceCurveNames[Index]);
+						if (UID != SmartName::MaxUID)
+						{
+							// has a valid UID, add to the list
+							SmartName::UID_Type& Value = CurveNameToUIDMap.Add(SourceCurveNames[Index]);
+							Value = UID;
+						}
 					}
 				}
 			}
