Improve tooltips in AnimDynamics AnimBP Node for Angular Spring.

diff --git a/Engine/Source/Runtime/AnimGraphRuntime/Public/BoneControllers/AnimNode_AnimDynamics.h b/Engine/Source/Runtime/AnimGraphRuntime/Public/BoneControllers/AnimNode_AnimDynamics.h
--- a/Engine/Source/Runtime/AnimGraphRuntime/Public/BoneControllers/AnimNode_AnimDynamics.h
+++ b/Engine/Source/Runtime/AnimGraphRuntime/Public/BoneControllers/AnimNode_AnimDynamics.h
@@ -140,7 +140,11 @@ struct FAnimPhysConstraintSetup
 	UPROPERTY(EditAnywhere, Category = Angular)
 	AnimPhysTwistAxis TwistAxis;
 
-	/** Axis on body1 to match to the angular target direction. */
+	/**
+	 * The axis in the simulation pose to align to the Angular Target.
+	 * This is typically the axis pointing along the bone.
+	 * Note: This is affected by the Angular Spring Constant.
+	 */
 	UPROPERTY(EditAnywhere, Category = Angular, meta=(DisplayAfter=AngularLimitsMax))
 	AnimPhysTwistAxis AngularTargetAxis;
 
@@ -168,7 +172,12 @@ struct FAnimPhysConstraintSetup
 	UPROPERTY(EditAnywhere, Category = Angular, meta = (UIMin = "-180", UIMax = "180", ClampMin = "-180", ClampMax = "180"))
 	FVector AngularLimitsMax;
 
-	/** Target direction to face for body1 (in body0 local space) */
+	/**
+	 * The axis to align the angular spring constraint to in the animation pose.
+	 * This typically points down the bone - so values of (1.0, 0.0, 0.0) are common,
+	 * but you can pick other values to align the spring to a different direction.
+	 * Note: This is affected by the Angular Spring Constant.
+	 */
 	UPROPERTY(EditAnywhere, Category = Angular)
 	FVector AngularTarget;
 };
@@ -232,11 +241,15 @@ struct ANIMGRAPHRUNTIME_API FAnimNode_AnimDynamics : public FAnimNode_SkeletalCo
 
 	FAnimNode_AnimDynamics();
 
-	/** Overridden linear damping value */
+	/** 
+	 * Overridden linear damping value. The default is 0.7. Values below 0.7 won't have an effect.
+	 */
 	UPROPERTY(EditAnywhere, AdvancedDisplay, BlueprintReadWrite, Category = Setup, meta = (PinHiddenByDefault, DisplayAfter="bOverrideLinearDamping"))
 	float LinearDampingOverride;
 
-	/** Overridden angular damping value */
+	/** 
+	 * Overridden angular damping value. The default is 0.7. Values below 0.7 won't have an effect.
+	 */
 	UPROPERTY(EditAnywhere, AdvancedDisplay, BlueprintReadWrite, Category = Setup, meta = (PinHiddenByDefault, DisplayAfter="bOverrideAngularDamping"))
 	float AngularDampingOverride;
 
@@ -276,7 +289,10 @@ struct ANIMGRAPHRUNTIME_API FAnimNode_AnimDynamics : public FAnimNode_SkeletalCo
 	UPROPERTY(EditAnywhere, BlueprintReadWrite, Category = Setup, meta = (PinHiddenByDefault, DisplayAfter="bAngularSpring"))
 	float LinearSpringConstant;
 
-	/** Spring constant to use when calculating angular springs, higher values mean a stronger spring */
+	/** 
+	 * Spring constant to use when calculating angular springs, higher values mean a stronger spring.
+	 * Note: Make sure to also set the Angular Target Axis and Angular Target in the Constraint Setup for this to have an effect.
+	 */
 	UPROPERTY(EditAnywhere, BlueprintReadWrite, Category = Setup, meta = (PinHiddenByDefault))
 	float AngularSpringConstant;
 
