Add option to simulate rigidbody animnodes in the frame in which they're reset (effectively skipping forward a frame).

diff --git a/Engine/Source/Runtime/AnimGraphRuntime/Private/BoneControllers/AnimNode_RigidBody.cpp b/Engine/Source/Runtime/AnimGraphRuntime/Private/BoneControllers/AnimNode_RigidBody.cpp
--- a/Engine/Source/Runtime/AnimGraphRuntime/Private/BoneControllers/AnimNode_RigidBody.cpp
+++ b/Engine/Source/Runtime/AnimGraphRuntime/Private/BoneControllers/AnimNode_RigidBody.cpp
@@ -254,7 +254,8 @@ void FAnimNode_RigidBody::EvaluateSkeletalControl_AnyThread(FComponentSpacePoseC
 
 		// If time advances, update simulation
 		// Reset if necessary
-		if (ResetSimulatedTeleportType != ETeleportType::None)
+		bool bDynamicsReset = (ResetSimulatedTeleportType != ETeleportType::None);
+		if (bDynamicsReset)
 		{
 			// Capture bone velocities if we have captured a bone velocity pose.
 			if (bTransferBoneVelocities && (CapturedBoneVelocityPose.GetPose().GetNumBones() > 0))
@@ -358,7 +359,7 @@ void FAnimNode_RigidBody::EvaluateSkeletalControl_AnyThread(FComponentSpacePoseC
 			PreviousComponentLinearVelocity = FVector::ZeroVector;
 		}
 		// Only need to tick physics if we didn't reset and we have some time to simulate
-		else if(DeltaSeconds > AnimPhysicsMinDeltaTime)
+		if((bSimulateAnimPhysicsAfterReset || !bDynamicsReset) && DeltaSeconds > AnimPhysicsMinDeltaTime)
 		{
 			// Transfer bone velocities previously captured.
 			if (bTransferBoneVelocities && (CapturedBoneVelocityPose.GetPose().GetNumBones() > 0))
@@ -591,10 +592,12 @@ void FAnimNode_RigidBody::InitPhysics(const UAnimInstance* InAnimInstance)
 	if (UPhysicsSettings* Settings = UPhysicsSettings::Get())
 	{
 		AnimPhysicsMinDeltaTime = Settings->AnimPhysicsMinDeltaTime;
+		bSimulateAnimPhysicsAfterReset = Settings->bSimulateAnimPhysicsAfterReset;
 	}
 	else
 	{
 		AnimPhysicsMinDeltaTime = 0.f;
+		bSimulateAnimPhysicsAfterReset = false;
 	}
 
 	if(UsePhysicsAsset)
diff --git a/Engine/Source/Runtime/AnimGraphRuntime/Public/BoneControllers/AnimNode_RigidBody.h b/Engine/Source/Runtime/AnimGraphRuntime/Public/BoneControllers/AnimNode_RigidBody.h
--- a/Engine/Source/Runtime/AnimGraphRuntime/Public/BoneControllers/AnimNode_RigidBody.h
+++ b/Engine/Source/Runtime/AnimGraphRuntime/Public/BoneControllers/AnimNode_RigidBody.h
@@ -175,6 +175,7 @@ private:
 
 	float AccumulatedDeltaTime;
 	float AnimPhysicsMinDeltaTime;
+	bool bSimulateAnimPhysicsAfterReset;
 	/** This should only be used for removing the delegate during termination. Do NOT use this for any per frame work */
 	TWeakObjectPtr<USkeletalMeshComponent> SkelMeshCompWeakPtr;
 
diff --git a/Engine/Source/Runtime/Engine/Classes/PhysicsEngine/PhysicsSettings.h b/Engine/Source/Runtime/Engine/Classes/PhysicsEngine/PhysicsSettings.h
--- a/Engine/Source/Runtime/Engine/Classes/PhysicsEngine/PhysicsSettings.h
+++ b/Engine/Source/Runtime/Engine/Classes/PhysicsEngine/PhysicsSettings.h
@@ -220,6 +220,10 @@ class ENGINE_API UPhysicsSettings : public UDeveloperSettings
 	UPROPERTY(config, EditAnywhere, meta = (ClampMin = "0.0", UIMin = "0.0", ClampMax = "1.0", UIMax = "1.0"), Category = Framerate)
 	float AnimPhysicsMinDeltaTime;
 
+	/** Whether to simulate anim physics nodes in the tick where they're reset. */
+	UPROPERTY(config, EditAnywhere, Category = Simulation)
+	bool bSimulateAnimPhysicsAfterReset;
+
 	/** Max Physics Delta Time to be clamped. */
 	UPROPERTY(config, EditAnywhere, meta=(ClampMin="0.0013", UIMin = "0.0013", ClampMax="1.0", UIMax="1.0"), Category=Framerate)
 	float MaxPhysicsDeltaTime;
diff --git a/Engine/Source/Runtime/Engine/Private/PhysicsEngine/PhysicsSettings.cpp b/Engine/Source/Runtime/Engine/Private/PhysicsEngine/PhysicsSettings.cpp
--- a/Engine/Source/Runtime/Engine/Private/PhysicsEngine/PhysicsSettings.cpp
+++ b/Engine/Source/Runtime/Engine/Private/PhysicsEngine/PhysicsSettings.cpp
@@ -31,6 +31,7 @@ UPhysicsSettings::UPhysicsSettings(const FObjectInitializer& ObjectInitializer)
 	, bDisableActiveActors(false)
 	, bEnableEnhancedDeterminism(false)
 	, AnimPhysicsMinDeltaTime(0.f)
+	, bSimulateAnimPhysicsAfterReset(false)
 	, MaxPhysicsDeltaTime(1.f / 30.f)
 	, bSubstepping(false)
 	, MaxSubstepDeltaTime(1.f / 60.f)
