Reinitialize the Copy Pose From Mesh node when the target mesh changes.
Fixes animation pose breaking when dynamically changing between skeletal meshes that use the same animbp but have a different number of bones.

diff --git a/Engine/Source/Runtime/AnimGraphRuntime/Private/AnimNodes/AnimNode_CopyPoseFromMesh.cpp b/Engine/Source/Runtime/AnimGraphRuntime/Private/AnimNodes/AnimNode_CopyPoseFromMesh.cpp
--- a/Engine/Source/Runtime/AnimGraphRuntime/Private/AnimNodes/AnimNode_CopyPoseFromMesh.cpp
+++ b/Engine/Source/Runtime/AnimGraphRuntime/Private/AnimNodes/AnimNode_CopyPoseFromMesh.cpp
@@ -47,6 +47,14 @@ void FAnimNode_CopyPoseFromMesh::RefreshMeshComponent(USkeletalMeshComponent* Ta
 			{
 				ReinitializeMeshComponent(InMeshComponent, InTargetMeshComponent);
 			}
+			else if (InTargetMeshComponent)
+			{
+				// see if target mesh has changed
+				if (InTargetMeshComponent->SkeletalMesh != CurrentlyUsedTargetMesh.Get())
+				{
+					ReinitializeMeshComponent(InMeshComponent, InTargetMeshComponent);
+				}
+			}
 		}
 		// if not valid, but input mesh is
 		else if (!CurrentMeshComponent && InMeshComponent)
@@ -205,6 +213,7 @@ void FAnimNode_CopyPoseFromMesh::ReinitializeMeshComponent(USkeletalMeshComponen
 	CurrentlyUsedSourceMeshComponent.Reset();
 	// reset source mesh
 	CurrentlyUsedSourceMesh.Reset();
+	CurrentlyUsedTargetMesh.Reset();
 	BoneMapToSource.Reset();
 	CurveNameToUIDMap.Reset();
 
@@ -218,6 +227,7 @@ void FAnimNode_CopyPoseFromMesh::ReinitializeMeshComponent(USkeletalMeshComponen
 		{
 			CurrentlyUsedSourceMeshComponent = NewSourceMeshComponent;
 			CurrentlyUsedSourceMesh = SourceSkelMesh;
+			CurrentlyUsedTargetMesh = TargetSkelMesh;
 
 			if (SourceSkelMesh == TargetSkelMesh)
 			{
diff --git a/Engine/Source/Runtime/AnimGraphRuntime/Public/AnimNodes/AnimNode_CopyPoseFromMesh.h b/Engine/Source/Runtime/AnimGraphRuntime/Public/AnimNodes/AnimNode_CopyPoseFromMesh.h
--- a/Engine/Source/Runtime/AnimGraphRuntime/Public/AnimNodes/AnimNode_CopyPoseFromMesh.h
+++ b/Engine/Source/Runtime/AnimGraphRuntime/Public/AnimNodes/AnimNode_CopyPoseFromMesh.h
@@ -44,9 +44,13 @@ struct ANIMGRAPHRUNTIME_API FAnimNode_CopyPoseFromMesh : public FAnimNode_Base
 	// End of FAnimNode_Base interface
 
 private:
+	// this is source mesh references, so that we could compare and see if it has changed
 	TWeakObjectPtr<USkeletalMeshComponent>	CurrentlyUsedSourceMeshComponent;
 	TWeakObjectPtr<USkeletalMesh>			CurrentlyUsedSourceMesh;
 	TWeakObjectPtr<USkeletalMesh>			CurrentlyUsedMesh;
+
+	// target mesh 
+	TWeakObjectPtr<USkeletalMesh>			CurrentlyUsedTargetMesh;
 	// cache of target space bases to source space bases
 	TMap<int32, int32> BoneMapToSource;
 	TMap<FName, SmartName::UID_Type> CurveNameToUIDMap;
