Workaround for CurvePoseIndex out of bounds crash
#rnx

diff --git a/Engine/Source/Runtime/AnimGraphRuntime/Private/AnimNodes/AnimNode_LayeredBoneBlend.cpp b/Engine/Source/Runtime/AnimGraphRuntime/Private/AnimNodes/AnimNode_LayeredBoneBlend.cpp
--- a/Engine/Source/Runtime/AnimGraphRuntime/Private/AnimNodes/AnimNode_LayeredBoneBlend.cpp
+++ b/Engine/Source/Runtime/AnimGraphRuntime/Private/AnimNodes/AnimNode_LayeredBoneBlend.cpp
@@ -246,7 +246,7 @@ void FAnimNode_LayeredBoneBlend::Evaluate_AnyThread(FPoseContext& Output)
 		for (int32 UIDIndex = 0; UIDIndex < CurveUIDFinder->Num(); ++UIDIndex)
 		{
 			int32 CurvePoseIndex = Output.Curve.GetArrayIndexByUID(UIDIndex);
-			if (CurvePoseIndex != INDEX_NONE)
+			if (CurvePoseSourceIndices.IsValidIndex(CurvePoseIndex))
 			{
 				int32 SourceIndex = CurvePoseSourceIndices[CurvePoseIndex];
 				if (SourceIndex != DEFAULT_SOURCEINDEX)
