Add asset info to AnimNode_Rigidbody.

diff --git a/Engine/Source/Runtime/AnimGraphRuntime/Private/BoneControllers/AnimNode_RigidBody.cpp b/Engine/Source/Runtime/AnimGraphRuntime/Private/BoneControllers/AnimNode_RigidBody.cpp
--- a/Engine/Source/Runtime/AnimGraphRuntime/Private/BoneControllers/AnimNode_RigidBody.cpp
+++ b/Engine/Source/Runtime/AnimGraphRuntime/Private/BoneControllers/AnimNode_RigidBody.cpp
@@ -29,6 +29,7 @@ FAnimNode_RigidBody::FAnimNode_RigidBody():
 	ResetSimulatedTeleportType = ETeleportType::None;
 	PhysicsSimulation = nullptr;
 	OverridePhysicsAsset = nullptr;
+	UsePhysicsAsset = nullptr;
 	bOverrideWorldGravity = false;
 	CachedBoundsScale = 1.2f;
 	SimulationSpace = ESimulationSpace::ComponentSpace;
@@ -217,6 +218,7 @@ DECLARE_CYCLE_STAT(TEXT("FAnimNode_Ragdoll::EvaluateSkeletalControl_AnyThread"),
 void FAnimNode_RigidBody::EvaluateSkeletalControl_AnyThread(FComponentSpacePoseContext& Output, TArray<FBoneTransform>& OutBoneTransforms)
 {
 	SCOPE_CYCLE_COUNTER(STAT_RigidBody_Eval);
+	FScopeCycleCounterUObject AdditionalScope(UsePhysicsAsset, GET_STATID(STAT_RigidBody_Eval));
 	SCOPE_CYCLE_COUNTER(STAT_ImmediateEvaluateSkeletalControl);
 	//SCOPED_NAMED_EVENT_TEXT("FAnimNode_Ragdoll::EvaluateSkeletalControl_AnyThread", FColor::Magenta);
 
@@ -573,7 +575,7 @@ void FAnimNode_RigidBody::InitPhysics(const UAnimInstance* InAnimInstance)
 	const USkeletalMesh* SkeletalMeshAsset = SkeletalMeshComp->SkeletalMesh;
 
 	const FReferenceSkeleton& SkelMeshRefSkel = SkeletalMeshAsset->RefSkeleton;
-	UPhysicsAsset* UsePhysicsAsset = OverridePhysicsAsset ? OverridePhysicsAsset : InAnimInstance->GetSkelMeshComponent()->GetPhysicsAsset();
+	UsePhysicsAsset = OverridePhysicsAsset ? OverridePhysicsAsset : InAnimInstance->GetSkelMeshComponent()->GetPhysicsAsset();
 		
 	USkeleton* SkeletonAsset = InAnimInstance->CurrentSkeleton;
 	ensure(SkeletonAsset == SkeletalMeshAsset->Skeleton);
diff --git a/Engine/Source/Runtime/AnimGraphRuntime/Public/BoneControllers/AnimNode_RigidBody.h b/Engine/Source/Runtime/AnimGraphRuntime/Public/BoneControllers/AnimNode_RigidBody.h
--- a/Engine/Source/Runtime/AnimGraphRuntime/Public/BoneControllers/AnimNode_RigidBody.h
+++ b/Engine/Source/Runtime/AnimGraphRuntime/Public/BoneControllers/AnimNode_RigidBody.h
@@ -69,6 +69,7 @@ private:
 	FTransform CurrentTransform;
 	FTransform PreviousTransform;
 
+	UPhysicsAsset* UsePhysicsAsset;
 public:
 	/** Override gravity*/
 	UPROPERTY(EditAnywhere, Category = Settings, meta = (PinHiddenByDefault, editcondition = "bOverrideWorldGravity"))
