Final copy-up for Control Rig Task Stream.

diff --git a/Engine/Source/Editor/UnrealEd/Private/Kismet2/Kismet2.cpp b/Engine/Source/Editor/UnrealEd/Private/Kismet2/Kismet2.cpp
--- a/Engine/Source/Editor/UnrealEd/Private/Kismet2/Kismet2.cpp
+++ b/Engine/Source/Editor/UnrealEd/Private/Kismet2/Kismet2.cpp
@@ -80,6 +80,7 @@
 #include "Framework/Notifications/NotificationManager.h"
 #include "Widgets/Notifications/SNotificationList.h"
 #include "Engine/InheritableComponentHandler.h"
+#include "Stats/StatsHierarchical.h"
 #include "Classes/EditorStyleSettings.h"
 
 DECLARE_CYCLE_STAT(TEXT("Compile Blueprint"), EKismetCompilerStats_CompileBlueprint, STATGROUP_KismetCompiler);
@@ -773,6 +774,8 @@ bool FKismetEditorUtilities::IsReferencedByUndoBuffer(UBlueprint* Blueprint)
 
 void FKismetEditorUtilities::CompileBlueprint(UBlueprint* BlueprintObj, EBlueprintCompileOptions CompileFlags, FCompilerResultsLog* pResults)
 {
+	DECLARE_SCOPE_HIERARCHICAL_COUNTER_FUNC()
+
 	if(GBlueprintUseCompilationManager)
 	{
 		FBlueprintCompilationManager::CompileSynchronously(FBPCompileRequest(BlueprintObj, CompileFlags, pResults));
diff --git a/Engine/Source/Runtime/AnimGraphRuntime/Private/AnimNodes/AnimNode_AimOffsetLookAt.cpp b/Engine/Source/Runtime/AnimGraphRuntime/Private/AnimNodes/AnimNode_AimOffsetLookAt.cpp
--- a/Engine/Source/Runtime/AnimGraphRuntime/Private/AnimNodes/AnimNode_AimOffsetLookAt.cpp
+++ b/Engine/Source/Runtime/AnimGraphRuntime/Private/AnimNodes/AnimNode_AimOffsetLookAt.cpp
@@ -18,6 +18,7 @@ TAutoConsoleVariable<int32> CVarAimOffsetLookAtDebug(TEXT("a.AnimNode.AimOffsetL
 
 void FAnimNode_AimOffsetLookAt::Initialize_AnyThread(const FAnimationInitializeContext& Context)
 {
+	DECLARE_SCOPE_HIERARCHICAL_COUNTER_ANIMNODE(Initialize_AnyThread)
 	FAnimNode_BlendSpacePlayer::Initialize_AnyThread(Context);
 	BasePose.Initialize(Context);
 }
@@ -66,6 +67,7 @@ void FAnimNode_AimOffsetLookAt::UpdateAssetPlayer(const FAnimationUpdateContext&
 
 void FAnimNode_AimOffsetLookAt::CacheBones_AnyThread(const FAnimationCacheBonesContext& Context)
 {
+	DECLARE_SCOPE_HIERARCHICAL_COUNTER_ANIMNODE(CacheBones_AnyThread)
 	FAnimNode_BlendSpacePlayer::CacheBones_AnyThread(Context);
 	BasePose.CacheBones(Context);
 
@@ -168,6 +170,7 @@ void FAnimNode_AimOffsetLookAt::UpdateFromLookAtTarget(FPoseContext& LocalPoseCo
 
 void FAnimNode_AimOffsetLookAt::GatherDebugData(FNodeDebugData& DebugData)
 {
+	DECLARE_SCOPE_HIERARCHICAL_COUNTER_ANIMNODE(GatherDebugData)
 	FString DebugLine = DebugData.GetNodeName(this);
 
 	DebugLine += FString::Printf(TEXT("(Play Time: %.3f)"), InternalTimeAccumulator);
diff --git a/Engine/Source/Runtime/AnimGraphRuntime/Private/AnimNodes/AnimNode_ApplyAdditive.cpp b/Engine/Source/Runtime/AnimGraphRuntime/Private/AnimNodes/AnimNode_ApplyAdditive.cpp
--- a/Engine/Source/Runtime/AnimGraphRuntime/Private/AnimNodes/AnimNode_ApplyAdditive.cpp
+++ b/Engine/Source/Runtime/AnimGraphRuntime/Private/AnimNodes/AnimNode_ApplyAdditive.cpp
@@ -9,6 +9,7 @@
 
 void FAnimNode_ApplyAdditive::Initialize_AnyThread(const FAnimationInitializeContext& Context)
 {
+	DECLARE_SCOPE_HIERARCHICAL_COUNTER_ANIMNODE(Initialize_AnyThread)
 	FAnimNode_Base::Initialize_AnyThread(Context);
 
 	Base.Initialize(Context);
@@ -18,14 +19,16 @@ void FAnimNode_ApplyAdditive::Initialize_AnyThread(const FAnimationInitializeCon
 	AlphaScaleBiasClamp.Reinitialize();
 }
 
-void FAnimNode_ApplyAdditive::CacheBones_AnyThread(const FAnimationCacheBonesContext& Context) 
+void FAnimNode_ApplyAdditive::CacheBones_AnyThread(const FAnimationCacheBonesContext& Context)
 {
+	DECLARE_SCOPE_HIERARCHICAL_COUNTER_ANIMNODE(CacheBones_AnyThread)
 	Base.CacheBones(Context);
 	Additive.CacheBones(Context);
 }
 
 void FAnimNode_ApplyAdditive::Update_AnyThread(const FAnimationUpdateContext& Context)
 {
+	DECLARE_SCOPE_HIERARCHICAL_COUNTER_ANIMNODE(Update_AnyThread)
 	Base.Update(Context);
 
 	ActualAlpha = 0.f;
@@ -60,6 +63,7 @@ void FAnimNode_ApplyAdditive::Update_AnyThread(const FAnimationUpdateContext& Co
 
 void FAnimNode_ApplyAdditive::Evaluate_AnyThread(FPoseContext& Output)
 {
+	DECLARE_SCOPE_HIERARCHICAL_COUNTER_ANIMNODE(Evaluate_AnyThread)
 	//@TODO: Could evaluate Base into Output and save a copy
 	if (FAnimWeight::IsRelevant(ActualAlpha))
 	{
@@ -90,6 +94,7 @@ FAnimNode_ApplyAdditive::FAnimNode_ApplyAdditive()
 
 void FAnimNode_ApplyAdditive::GatherDebugData(FNodeDebugData& DebugData)
 {
+	DECLARE_SCOPE_HIERARCHICAL_COUNTER_ANIMNODE(GatherDebugData)
 	FString DebugLine = DebugData.GetNodeName(this);
 	DebugLine += FString::Printf(TEXT("(Alpha: %.1f%%)"), ActualAlpha*100.f);
 
diff --git a/Engine/Source/Runtime/AnimGraphRuntime/Private/AnimNodes/AnimNode_BlendBoneByChannel.cpp b/Engine/Source/Runtime/AnimGraphRuntime/Private/AnimNodes/AnimNode_BlendBoneByChannel.cpp
--- a/Engine/Source/Runtime/AnimGraphRuntime/Private/AnimNodes/AnimNode_BlendBoneByChannel.cpp
+++ b/Engine/Source/Runtime/AnimGraphRuntime/Private/AnimNodes/AnimNode_BlendBoneByChannel.cpp
@@ -9,6 +9,7 @@
 
 void FAnimNode_BlendBoneByChannel::Initialize_AnyThread(const FAnimationInitializeContext& Context)
 {
+	DECLARE_SCOPE_HIERARCHICAL_COUNTER_ANIMNODE(Initialize_AnyThread)
 	FAnimNode_Base::Initialize_AnyThread(Context);
 
 	A.Initialize(Context);
@@ -17,6 +18,7 @@ void FAnimNode_BlendBoneByChannel::Initialize_AnyThread(const FAnimationInitiali
 
 void FAnimNode_BlendBoneByChannel::CacheBones_AnyThread(const FAnimationCacheBonesContext& Context)
 {
+	DECLARE_SCOPE_HIERARCHICAL_COUNTER_ANIMNODE(CacheBones_AnyThread)
 	A.CacheBones(Context);
 	B.CacheBones(Context);
 
@@ -38,6 +40,7 @@ void FAnimNode_BlendBoneByChannel::CacheBones_AnyThread(const FAnimationCacheBon
 
 void FAnimNode_BlendBoneByChannel::Update_AnyThread(const FAnimationUpdateContext& Context)
 {
+	DECLARE_SCOPE_HIERARCHICAL_COUNTER_ANIMNODE(Update_AnyThread)
 	QUICK_SCOPE_CYCLE_COUNTER(STAT_FAnimNode_BlendBoneByChannel_Update);
 	GetEvaluateGraphExposedInputs().Execute(Context);
 
@@ -53,6 +56,7 @@ void FAnimNode_BlendBoneByChannel::Update_AnyThread(const FAnimationUpdateContex
 
 void FAnimNode_BlendBoneByChannel::Evaluate_AnyThread(FPoseContext& Output)
 {
+	DECLARE_SCOPE_HIERARCHICAL_COUNTER_ANIMNODE(Evaluate_AnyThread)
 	A.Evaluate(Output);
 
 	if (bBIsRelevant)
@@ -171,6 +175,7 @@ void FAnimNode_BlendBoneByChannel::Evaluate_AnyThread(FPoseContext& Output)
 
 void FAnimNode_BlendBoneByChannel::GatherDebugData(FNodeDebugData& DebugData)
 {
+	DECLARE_SCOPE_HIERARCHICAL_COUNTER_ANIMNODE(GatherDebugData)
 	FString DebugLine = DebugData.GetNodeName(this);
 	DebugLine += FString::Printf(TEXT("(Alpha: %.1f%%)"), InternalBlendAlpha * 100);
 	DebugData.AddDebugItem(DebugLine);
diff --git a/Engine/Source/Runtime/AnimGraphRuntime/Private/AnimNodes/AnimNode_BlendListBase.cpp b/Engine/Source/Runtime/AnimGraphRuntime/Private/AnimNodes/AnimNode_BlendListBase.cpp
--- a/Engine/Source/Runtime/AnimGraphRuntime/Private/AnimNodes/AnimNode_BlendListBase.cpp
+++ b/Engine/Source/Runtime/AnimGraphRuntime/Private/AnimNodes/AnimNode_BlendListBase.cpp
@@ -9,6 +9,7 @@
 
 void FAnimNode_BlendListBase::Initialize_AnyThread(const FAnimationInitializeContext& Context)
 {
+	DECLARE_SCOPE_HIERARCHICAL_COUNTER_ANIMNODE(Initialize_AnyThread)
 	FAnimNode_Base::Initialize_AnyThread(Context);
 
 	const int NumPoses = BlendPose.Num();
@@ -63,8 +64,9 @@ void FAnimNode_BlendListBase::Initialize_AnyThread(const FAnimationInitializeCon
 	}
 }
 
-void FAnimNode_BlendListBase::CacheBones_AnyThread(const FAnimationCacheBonesContext& Context) 
+void FAnimNode_BlendListBase::CacheBones_AnyThread(const FAnimationCacheBonesContext& Context)
 {
+	DECLARE_SCOPE_HIERARCHICAL_COUNTER_ANIMNODE(CacheBones_AnyThread)
 	for(int32 ChildIndex=0; ChildIndex<BlendPose.Num(); ChildIndex++)
 	{
 		BlendPose[ChildIndex].CacheBones(Context);
@@ -73,6 +75,7 @@ void FAnimNode_BlendListBase::CacheBones_AnyThread(const FAnimationCacheBonesCon
 
 void FAnimNode_BlendListBase::Update_AnyThread(const FAnimationUpdateContext& Context)
 {
+	DECLARE_SCOPE_HIERARCHICAL_COUNTER_ANIMNODE(Update_AnyThread)
 	GetEvaluateGraphExposedInputs().Execute(Context);
 
 	const int NumPoses = BlendPose.Num();
@@ -203,6 +206,7 @@ void FAnimNode_BlendListBase::Update_AnyThread(const FAnimationUpdateContext& Co
 
 void FAnimNode_BlendListBase::Evaluate_AnyThread(FPoseContext& Output)
 {
+	DECLARE_SCOPE_HIERARCHICAL_COUNTER_ANIMNODE(Evaluate_AnyThread)
 	ANIM_MT_SCOPE_CYCLE_COUNTER(BlendPosesInGraph, !IsInGameThread());
 
 	const int32 NumPoses = PosesToEvaluate.Num();
@@ -247,6 +251,7 @@ void FAnimNode_BlendListBase::Evaluate_AnyThread(FPoseContext& Output)
 
 void FAnimNode_BlendListBase::GatherDebugData(FNodeDebugData& DebugData)
 {
+	DECLARE_SCOPE_HIERARCHICAL_COUNTER_ANIMNODE(GatherDebugData)
 	const int NumPoses = BlendPose.Num();
 	const int32 ChildIndex = GetActiveChildIndex();
 
diff --git a/Engine/Source/Runtime/AnimGraphRuntime/Private/AnimNodes/AnimNode_BlendSpaceEvaluator.cpp b/Engine/Source/Runtime/AnimGraphRuntime/Private/AnimNodes/AnimNode_BlendSpaceEvaluator.cpp
--- a/Engine/Source/Runtime/AnimGraphRuntime/Private/AnimNodes/AnimNode_BlendSpaceEvaluator.cpp
+++ b/Engine/Source/Runtime/AnimGraphRuntime/Private/AnimNodes/AnimNode_BlendSpaceEvaluator.cpp
@@ -24,6 +24,7 @@ void FAnimNode_BlendSpaceEvaluator::UpdateAssetPlayer(const FAnimationUpdateCont
 
 void FAnimNode_BlendSpaceEvaluator::GatherDebugData(FNodeDebugData& DebugData)
 {
+	DECLARE_SCOPE_HIERARCHICAL_COUNTER_ANIMNODE(GatherDebugData)
 	FString DebugLine = DebugData.GetNodeName(this);
 
 	DebugLine += FString::Printf(TEXT("('%s' Play Time: %.3f)"), *BlendSpace->GetName(), InternalTimeAccumulator);
diff --git a/Engine/Source/Runtime/AnimGraphRuntime/Private/AnimNodes/AnimNode_BlendSpacePlayer.cpp b/Engine/Source/Runtime/AnimGraphRuntime/Private/AnimNodes/AnimNode_BlendSpacePlayer.cpp
--- a/Engine/Source/Runtime/AnimGraphRuntime/Private/AnimNodes/AnimNode_BlendSpacePlayer.cpp
+++ b/Engine/Source/Runtime/AnimGraphRuntime/Private/AnimNodes/AnimNode_BlendSpacePlayer.cpp
@@ -55,6 +55,7 @@ float FAnimNode_BlendSpacePlayer::GetCurrentAssetLength()
 
 void FAnimNode_BlendSpacePlayer::Initialize_AnyThread(const FAnimationInitializeContext& Context)
 {
+	DECLARE_SCOPE_HIERARCHICAL_COUNTER_ANIMNODE(Initialize_AnyThread)
 	FAnimNode_AssetPlayerBase::Initialize_AnyThread(Context);
 
 	GetEvaluateGraphExposedInputs().Execute(Context);
@@ -64,8 +65,9 @@ void FAnimNode_BlendSpacePlayer::Initialize_AnyThread(const FAnimationInitialize
 	PreviousBlendSpace = BlendSpace;
 }
 
-void FAnimNode_BlendSpacePlayer::CacheBones_AnyThread(const FAnimationCacheBonesContext& Context) 
+void FAnimNode_BlendSpacePlayer::CacheBones_AnyThread(const FAnimationCacheBonesContext& Context)
 {
+	DECLARE_SCOPE_HIERARCHICAL_COUNTER_ANIMNODE(CacheBones_AnyThread)
 }
 
 void FAnimNode_BlendSpacePlayer::UpdateAssetPlayer(const FAnimationUpdateContext& Context)
@@ -77,6 +79,7 @@ void FAnimNode_BlendSpacePlayer::UpdateAssetPlayer(const FAnimationUpdateContext
 
 void FAnimNode_BlendSpacePlayer::UpdateInternal(const FAnimationUpdateContext& Context)
 {
+	DECLARE_SCOPE_HIERARCHICAL_COUNTER_ANIMNODE(UpdateInternal)
 	if ((BlendSpace != NULL) && (Context.AnimInstanceProxy->IsSkeletonCompatible(BlendSpace->GetSkeleton())))
 	{
 		// Create a tick record and fill it out
@@ -104,6 +107,7 @@ void FAnimNode_BlendSpacePlayer::UpdateInternal(const FAnimationUpdateContext& C
 
 void FAnimNode_BlendSpacePlayer::Evaluate_AnyThread(FPoseContext& Output)
 {
+	DECLARE_SCOPE_HIERARCHICAL_COUNTER_ANIMNODE(Evaluate_AnyThread)
 	if ((BlendSpace != NULL) && (Output.AnimInstanceProxy->IsSkeletonCompatible(BlendSpace->GetSkeleton())))
 	{
 		BlendSpace->GetAnimationPose(BlendSampleDataCache, Output.Pose, Output.Curve);
@@ -124,6 +128,7 @@ void FAnimNode_BlendSpacePlayer::OverrideAsset(UAnimationAsset* NewAsset)
 
 void FAnimNode_BlendSpacePlayer::GatherDebugData(FNodeDebugData& DebugData)
 {
+	DECLARE_SCOPE_HIERARCHICAL_COUNTER_ANIMNODE(GatherDebugData)
 	FString DebugLine = DebugData.GetNodeName(this);
 	if (BlendSpace)
 	{
diff --git a/Engine/Source/Runtime/AnimGraphRuntime/Private/AnimNodes/AnimNode_CopyPoseFromMesh.cpp b/Engine/Source/Runtime/AnimGraphRuntime/Private/AnimNodes/AnimNode_CopyPoseFromMesh.cpp
--- a/Engine/Source/Runtime/AnimGraphRuntime/Private/AnimNodes/AnimNode_CopyPoseFromMesh.cpp
+++ b/Engine/Source/Runtime/AnimGraphRuntime/Private/AnimNodes/AnimNode_CopyPoseFromMesh.cpp
@@ -15,6 +15,7 @@ FAnimNode_CopyPoseFromMesh::FAnimNode_CopyPoseFromMesh()
 
 void FAnimNode_CopyPoseFromMesh::Initialize_AnyThread(const FAnimationInitializeContext& Context)
 {
+	DECLARE_SCOPE_HIERARCHICAL_COUNTER_ANIMNODE(Initialize_AnyThread)
 	FAnimNode_Base::Initialize_AnyThread(Context);
 
 	// Initial update of the node, so we dont have a frame-delay on setup
@@ -26,6 +27,7 @@ void FAnimNode_CopyPoseFromMesh::Initialize_AnyThread(const FAnimationInitialize
 
 void FAnimNode_CopyPoseFromMesh::CacheBones_AnyThread(const FAnimationCacheBonesContext& Context)
 {
+	DECLARE_SCOPE_HIERARCHICAL_COUNTER_ANIMNODE(CacheBones_AnyThread)
 
 }
 
@@ -147,6 +149,7 @@ void FAnimNode_CopyPoseFromMesh::PreUpdate(const UAnimInstance* InAnimInstance)
 
 void FAnimNode_CopyPoseFromMesh::Update_AnyThread(const FAnimationUpdateContext& Context)
 {
+	DECLARE_SCOPE_HIERARCHICAL_COUNTER_ANIMNODE(Update_AnyThread)
 	// This introduces a frame of latency in setting the pin-driven source component,
 	// but we cannot do the work to extract transforms on a worker thread as it is not thread safe.
 	GetEvaluateGraphExposedInputs().Execute(Context);
@@ -154,6 +157,7 @@ void FAnimNode_CopyPoseFromMesh::Update_AnyThread(const FAnimationUpdateContext&
 
 void FAnimNode_CopyPoseFromMesh::Evaluate_AnyThread(FPoseContext& Output)
 {
+	DECLARE_SCOPE_HIERARCHICAL_COUNTER_ANIMNODE(Evaluate_AnyThread)
 	FCompactPose& OutPose = Output.Pose;
 	OutPose.ResetToRefPose();
 	USkeletalMesh* CurrentMesh = CurrentlyUsedMesh.IsValid() ? CurrentlyUsedMesh.Get() : nullptr;
@@ -202,6 +206,7 @@ void FAnimNode_CopyPoseFromMesh::Evaluate_AnyThread(FPoseContext& Output)
 
 void FAnimNode_CopyPoseFromMesh::GatherDebugData(FNodeDebugData& DebugData)
 {
+	DECLARE_SCOPE_HIERARCHICAL_COUNTER_ANIMNODE(GatherDebugData)
 	FString DebugLine = DebugData.GetNodeName(this);
 
 	DebugLine += FString::Printf(TEXT("('%s')"), *GetNameSafe(CurrentlyUsedSourceMeshComponent.IsValid() ? CurrentlyUsedSourceMeshComponent.Get()->SkeletalMesh : nullptr));
diff --git a/Engine/Source/Runtime/AnimGraphRuntime/Private/AnimNodes/AnimNode_CurveSource.cpp b/Engine/Source/Runtime/AnimGraphRuntime/Private/AnimNodes/AnimNode_CurveSource.cpp
--- a/Engine/Source/Runtime/AnimGraphRuntime/Private/AnimNodes/AnimNode_CurveSource.cpp
+++ b/Engine/Source/Runtime/AnimGraphRuntime/Private/AnimNodes/AnimNode_CurveSource.cpp
@@ -72,6 +72,7 @@ public:
 
 void FAnimNode_CurveSource::Evaluate_AnyThread(FPoseContext& Output)
 {
+	DECLARE_SCOPE_HIERARCHICAL_COUNTER_ANIMNODE(Evaluate_AnyThread)
 	SourcePose.Evaluate(Output);
 
 	if (CurveSource.GetInterface() != nullptr)
@@ -97,6 +98,7 @@ void FAnimNode_CurveSource::Evaluate_AnyThread(FPoseContext& Output)
 
 void FAnimNode_CurveSource::Update_AnyThread(const FAnimationUpdateContext& Context)
 {
+	DECLARE_SCOPE_HIERARCHICAL_COUNTER_ANIMNODE(Update_AnyThread)
 	// Evaluate any BP logic plugged into this node
 	GetEvaluateGraphExposedInputs().Execute(Context);
 	SourcePose.Update(Context);
@@ -104,18 +106,21 @@ void FAnimNode_CurveSource::Update_AnyThread(const FAnimationUpdateContext& Cont
 
 void FAnimNode_CurveSource::Initialize_AnyThread(const FAnimationInitializeContext& Context)
 {
+	DECLARE_SCOPE_HIERARCHICAL_COUNTER_ANIMNODE(Initialize_AnyThread)
 	FAnimNode_Base::Initialize_AnyThread(Context);
 	SourcePose.Initialize(Context);
 }
 
 void FAnimNode_CurveSource::CacheBones_AnyThread(const FAnimationCacheBonesContext& Context)
 {
+	DECLARE_SCOPE_HIERARCHICAL_COUNTER_ANIMNODE(CacheBones_AnyThread)
 	FAnimNode_Base::CacheBones_AnyThread(Context);
 	SourcePose.CacheBones(Context);
 }
 
 void FAnimNode_CurveSource::GatherDebugData(FNodeDebugData& DebugData)
 {
+	DECLARE_SCOPE_HIERARCHICAL_COUNTER_ANIMNODE(GatherDebugData)
 	FAnimNode_Base::GatherDebugData(DebugData);
 	SourcePose.GatherDebugData(DebugData.BranchFlow(1.f));
 }
\ No newline at end of file
diff --git a/Engine/Source/Runtime/AnimGraphRuntime/Private/AnimNodes/AnimNode_LayeredBoneBlend.cpp b/Engine/Source/Runtime/AnimGraphRuntime/Private/AnimNodes/AnimNode_LayeredBoneBlend.cpp
--- a/Engine/Source/Runtime/AnimGraphRuntime/Private/AnimNodes/AnimNode_LayeredBoneBlend.cpp
+++ b/Engine/Source/Runtime/AnimGraphRuntime/Private/AnimNodes/AnimNode_LayeredBoneBlend.cpp
@@ -10,6 +10,7 @@
 
 void FAnimNode_LayeredBoneBlend::Initialize_AnyThread(const FAnimationInitializeContext& Context)
 {
+	DECLARE_SCOPE_HIERARCHICAL_COUNTER_ANIMNODE(Initialize_AnyThread)
 	FAnimNode_Base::Initialize_AnyThread(Context);
 
 	const int NumPoses = BlendPoses.Num();
@@ -113,8 +114,9 @@ void FAnimNode_LayeredBoneBlend::ReinitializeBoneBlendWeights(const FBoneContain
 	}
 }
 
-void FAnimNode_LayeredBoneBlend::CacheBones_AnyThread(const FAnimationCacheBonesContext& Context) 
+void FAnimNode_LayeredBoneBlend::CacheBones_AnyThread(const FAnimationCacheBonesContext& Context)
 {
+	DECLARE_SCOPE_HIERARCHICAL_COUNTER_ANIMNODE(CacheBones_AnyThread)
 	BasePose.CacheBones(Context);
 	int32 NumPoses = BlendPoses.Num();
 	for(int32 ChildIndex=0; ChildIndex<NumPoses; ChildIndex++)
@@ -130,6 +132,7 @@ void FAnimNode_LayeredBoneBlend::CacheBones_AnyThread(const FAnimationCacheBones
 
 void FAnimNode_LayeredBoneBlend::Update_AnyThread(const FAnimationUpdateContext& Context)
 {
+	DECLARE_SCOPE_HIERARCHICAL_COUNTER_ANIMNODE(Update_AnyThread)
 	bHasRelevantPoses = false;
 	int32 RootMotionBlendPose = -1;
 	float RootMotionWeight = 0.f;
@@ -202,6 +205,7 @@ void FAnimNode_LayeredBoneBlend::Update_AnyThread(const FAnimationUpdateContext&
 
 void FAnimNode_LayeredBoneBlend::Evaluate_AnyThread(FPoseContext& Output)
 {
+	DECLARE_SCOPE_HIERARCHICAL_COUNTER_ANIMNODE(Evaluate_AnyThread)
 	ANIM_MT_SCOPE_CYCLE_COUNTER(BlendPosesInGraph, !IsInGameThread());
 
 	const int NumPoses = BlendPoses.Num();
@@ -281,6 +285,7 @@ void FAnimNode_LayeredBoneBlend::Evaluate_AnyThread(FPoseContext& Output)
 
 void FAnimNode_LayeredBoneBlend::GatherDebugData(FNodeDebugData& DebugData)
 {
+	DECLARE_SCOPE_HIERARCHICAL_COUNTER_ANIMNODE(GatherDebugData)
 	const int NumPoses = BlendPoses.Num();
 
 	FString DebugLine = DebugData.GetNodeName(this);
diff --git a/Engine/Source/Runtime/AnimGraphRuntime/Private/AnimNodes/AnimNode_ModifyCurve.cpp b/Engine/Source/Runtime/AnimGraphRuntime/Private/AnimNodes/AnimNode_ModifyCurve.cpp
--- a/Engine/Source/Runtime/AnimGraphRuntime/Private/AnimNodes/AnimNode_ModifyCurve.cpp
+++ b/Engine/Source/Runtime/AnimGraphRuntime/Private/AnimNodes/AnimNode_ModifyCurve.cpp
@@ -12,6 +12,7 @@ FAnimNode_ModifyCurve::FAnimNode_ModifyCurve()
 
 void FAnimNode_ModifyCurve::Initialize_AnyThread(const FAnimationInitializeContext& Context)
 {
+	DECLARE_SCOPE_HIERARCHICAL_COUNTER_ANIMNODE(Initialize_AnyThread)
 	Super::Initialize_AnyThread(Context);
 	SourcePose.Initialize(Context);
 
@@ -25,12 +26,14 @@ void FAnimNode_ModifyCurve::Initialize_AnyThread(const FAnimationInitializeConte
 
 void FAnimNode_ModifyCurve::CacheBones_AnyThread(const FAnimationCacheBonesContext& Context)
 {
+	DECLARE_SCOPE_HIERARCHICAL_COUNTER_ANIMNODE(CacheBones_AnyThread)
 	Super::CacheBones_AnyThread(Context);
 	SourcePose.CacheBones(Context);
 }
 
 void FAnimNode_ModifyCurve::Evaluate_AnyThread(FPoseContext& Output)
 {
+	DECLARE_SCOPE_HIERARCHICAL_COUNTER_ANIMNODE(Evaluate_AnyThread)
 	FPoseContext SourceData(Output);
 	SourcePose.Evaluate(SourceData);
 
@@ -87,6 +90,7 @@ void FAnimNode_ModifyCurve::Evaluate_AnyThread(FPoseContext& Output)
 
 void FAnimNode_ModifyCurve::Update_AnyThread(const FAnimationUpdateContext& Context)
 {
+	DECLARE_SCOPE_HIERARCHICAL_COUNTER_ANIMNODE(Update_AnyThread)
 	// Run update on input pose nodes
 	SourcePose.Update(Context);
 
diff --git a/Engine/Source/Runtime/AnimGraphRuntime/Private/AnimNodes/AnimNode_MultiWayBlend.cpp b/Engine/Source/Runtime/AnimGraphRuntime/Private/AnimNodes/AnimNode_MultiWayBlend.cpp
--- a/Engine/Source/Runtime/AnimGraphRuntime/Private/AnimNodes/AnimNode_MultiWayBlend.cpp
+++ b/Engine/Source/Runtime/AnimGraphRuntime/Private/AnimNodes/AnimNode_MultiWayBlend.cpp
@@ -16,6 +16,7 @@ struct FMultiBlendData : public TThreadSingleton<FMultiBlendData>
 
 void FAnimNode_MultiWayBlend::Initialize_AnyThread(const FAnimationInitializeContext& Context)
 {
+	DECLARE_SCOPE_HIERARCHICAL_COUNTER_ANIMNODE(Initialize_AnyThread)
 	FAnimNode_Base::Initialize_AnyThread(Context);
 
 	// this should be consistent all the time by editor node
@@ -32,8 +33,9 @@ void FAnimNode_MultiWayBlend::Initialize_AnyThread(const FAnimationInitializeCon
 	}
 }
 
-void FAnimNode_MultiWayBlend::CacheBones_AnyThread(const FAnimationCacheBonesContext& Context) 
+void FAnimNode_MultiWayBlend::CacheBones_AnyThread(const FAnimationCacheBonesContext& Context)
 {
+	DECLARE_SCOPE_HIERARCHICAL_COUNTER_ANIMNODE(CacheBones_AnyThread)
 	for (FPoseLink& Pose : Poses)
 	{
 		Pose.CacheBones(Context);
@@ -85,6 +87,7 @@ void FAnimNode_MultiWayBlend::UpdateCachedAlphas()
 
 void FAnimNode_MultiWayBlend::Update_AnyThread(const FAnimationUpdateContext& Context)
 {
+	DECLARE_SCOPE_HIERARCHICAL_COUNTER_ANIMNODE(Update_AnyThread)
 	QUICK_SCOPE_CYCLE_COUNTER(STAT_FAnimationNode_MultiWayBlend_Update);
 	GetEvaluateGraphExposedInputs().Execute(Context);
 	UpdateCachedAlphas();
@@ -102,6 +105,7 @@ void FAnimNode_MultiWayBlend::Update_AnyThread(const FAnimationUpdateContext& Co
 
 void FAnimNode_MultiWayBlend::Evaluate_AnyThread(FPoseContext& Output)
 {
+	DECLARE_SCOPE_HIERARCHICAL_COUNTER_ANIMNODE(Evaluate_AnyThread)
 	FMultiBlendData& BlendData = FMultiBlendData::Get();
 	TArray<FCompactPose, TInlineAllocator<8>>& SourcePoses = BlendData.SourcePoses;
 	TArray<float, TInlineAllocator<8>>& SourceWeights = BlendData.SourceWeights;
@@ -154,6 +158,7 @@ void FAnimNode_MultiWayBlend::Evaluate_AnyThread(FPoseContext& Output)
 
 void FAnimNode_MultiWayBlend::GatherDebugData(FNodeDebugData& DebugData)
 {
+	DECLARE_SCOPE_HIERARCHICAL_COUNTER_ANIMNODE(GatherDebugData)
 	FString DebugLine = DebugData.GetNodeName(this);
 	DebugData.AddDebugItem(DebugLine);
 
diff --git a/Engine/Source/Runtime/AnimGraphRuntime/Private/AnimNodes/AnimNode_PoseBlendNode.cpp b/Engine/Source/Runtime/AnimGraphRuntime/Private/AnimNodes/AnimNode_PoseBlendNode.cpp
--- a/Engine/Source/Runtime/AnimGraphRuntime/Private/AnimNodes/AnimNode_PoseBlendNode.cpp
+++ b/Engine/Source/Runtime/AnimGraphRuntime/Private/AnimNodes/AnimNode_PoseBlendNode.cpp
@@ -15,13 +15,15 @@ FAnimNode_PoseBlendNode::FAnimNode_PoseBlendNode()
 
 void FAnimNode_PoseBlendNode::Initialize_AnyThread(const FAnimationInitializeContext& Context)
 {
+	DECLARE_SCOPE_HIERARCHICAL_COUNTER_ANIMNODE(Initialize_AnyThread)
 	FAnimNode_PoseHandler::Initialize_AnyThread(Context);
 
 	SourcePose.Initialize(Context);
 }
 
-void FAnimNode_PoseBlendNode::CacheBones_AnyThread(const FAnimationCacheBonesContext& Context) 
+void FAnimNode_PoseBlendNode::CacheBones_AnyThread(const FAnimationCacheBonesContext& Context)
 {
+	DECLARE_SCOPE_HIERARCHICAL_COUNTER_ANIMNODE(CacheBones_AnyThread)
 	FAnimNode_PoseHandler::CacheBones_AnyThread(Context);
 	SourcePose.CacheBones(Context);
 }
@@ -34,6 +36,7 @@ void FAnimNode_PoseBlendNode::UpdateAssetPlayer(const FAnimationUpdateContext& C
 
 void FAnimNode_PoseBlendNode::Evaluate_AnyThread(FPoseContext& Output)
 {
+	DECLARE_SCOPE_HIERARCHICAL_COUNTER_ANIMNODE(Evaluate_AnyThread)
 	ANIM_MT_SCOPE_CYCLE_COUNTER(PoseBlendNodeEvaluate, !IsInGameThread());
 
 	FPoseContext SourceData(Output);
@@ -83,6 +86,7 @@ void FAnimNode_PoseBlendNode::Evaluate_AnyThread(FPoseContext& Output)
 
 void FAnimNode_PoseBlendNode::GatherDebugData(FNodeDebugData& DebugData)
 {
+	DECLARE_SCOPE_HIERARCHICAL_COUNTER_ANIMNODE(GatherDebugData)
 	FAnimNode_PoseHandler::GatherDebugData(DebugData);
 	SourcePose.GatherDebugData(DebugData.BranchFlow(1.f));
 }
diff --git a/Engine/Source/Runtime/AnimGraphRuntime/Private/AnimNodes/AnimNode_PoseByName.cpp b/Engine/Source/Runtime/AnimGraphRuntime/Private/AnimNodes/AnimNode_PoseByName.cpp
--- a/Engine/Source/Runtime/AnimGraphRuntime/Private/AnimNodes/AnimNode_PoseByName.cpp
+++ b/Engine/Source/Runtime/AnimGraphRuntime/Private/AnimNodes/AnimNode_PoseByName.cpp
@@ -8,6 +8,7 @@
 
 void FAnimNode_PoseByName::Initialize_AnyThread(const FAnimationInitializeContext& Context)
 {
+	DECLARE_SCOPE_HIERARCHICAL_COUNTER_ANIMNODE(Initialize_AnyThread)
 	FAnimNode_PoseHandler::Initialize_AnyThread(Context);
 }
 
@@ -38,6 +39,7 @@ void FAnimNode_PoseByName::UpdateAssetPlayer(const FAnimationUpdateContext& Cont
 
 void FAnimNode_PoseByName::Evaluate_AnyThread(FPoseContext& Output)
 {
+	DECLARE_SCOPE_HIERARCHICAL_COUNTER_ANIMNODE(Evaluate_AnyThread)
 	// make sure we have curve to eval
 	if ((CurrentPoseAsset.IsValid()) && (PoseExtractContext.PoseCurves.Num() > 0) && (Output.AnimInstanceProxy->IsSkeletonCompatible(CurrentPoseAsset->GetSkeleton())))
 	{
@@ -54,6 +56,7 @@ void FAnimNode_PoseByName::Evaluate_AnyThread(FPoseContext& Output)
 
 void FAnimNode_PoseByName::GatherDebugData(FNodeDebugData& DebugData)
 {
+	DECLARE_SCOPE_HIERARCHICAL_COUNTER_ANIMNODE(GatherDebugData)
 	FString DebugLine = DebugData.GetNodeName(this);
 	
 	DebugLine += FString::Printf(TEXT("('%s' Pose: %s)"), CurrentPoseAsset.IsValid()? *CurrentPoseAsset.Get()->GetName() : TEXT("None"), *PoseName.ToString());
diff --git a/Engine/Source/Runtime/AnimGraphRuntime/Private/AnimNodes/AnimNode_PoseDriver.cpp b/Engine/Source/Runtime/AnimGraphRuntime/Private/AnimNodes/AnimNode_PoseDriver.cpp
--- a/Engine/Source/Runtime/AnimGraphRuntime/Private/AnimNodes/AnimNode_PoseDriver.cpp
+++ b/Engine/Source/Runtime/AnimGraphRuntime/Private/AnimNodes/AnimNode_PoseDriver.cpp
@@ -23,6 +23,7 @@ FAnimNode_PoseDriver::FAnimNode_PoseDriver()
 
 void FAnimNode_PoseDriver::Initialize_AnyThread(const FAnimationInitializeContext& Context)
 {
+	DECLARE_SCOPE_HIERARCHICAL_COUNTER_ANIMNODE(Initialize_AnyThread)
 	FAnimNode_PoseHandler::Initialize_AnyThread(Context);
 
 	SourcePose.Initialize(Context);
@@ -71,6 +72,7 @@ void FAnimNode_PoseDriver::RebuildPoseList(const FBoneContainer& InBoneContainer
 
 void FAnimNode_PoseDriver::CacheBones_AnyThread(const FAnimationCacheBonesContext& Context)
 {
+	DECLARE_SCOPE_HIERARCHICAL_COUNTER_ANIMNODE(CacheBones_AnyThread)
 	FAnimNode_PoseHandler::CacheBones_AnyThread(Context);
 	// Init pose input
 	SourcePose.CacheBones(Context);
@@ -140,6 +142,7 @@ void FAnimNode_PoseDriver::UpdateAssetPlayer(const FAnimationUpdateContext& Cont
 
 void FAnimNode_PoseDriver::GatherDebugData(FNodeDebugData& DebugData)
 {
+	DECLARE_SCOPE_HIERARCHICAL_COUNTER_ANIMNODE(GatherDebugData)
 	FAnimNode_PoseHandler::GatherDebugData(DebugData);
 	SourcePose.GatherDebugData(DebugData.BranchFlow(1.f));
 }
@@ -215,6 +218,7 @@ void FAnimNode_PoseDriver::GetRBFTargets(TArray<FRBFTarget>& OutTargets) const
 
 void FAnimNode_PoseDriver::Evaluate_AnyThread(FPoseContext& Output)
 {
+	DECLARE_SCOPE_HIERARCHICAL_COUNTER_ANIMNODE(Evaluate_AnyThread)
 	QUICK_SCOPE_CYCLE_COUNTER(STAT_PoseDriver_Eval);
 
 	// Udpate DrivenIDs if needed
diff --git a/Engine/Source/Runtime/AnimGraphRuntime/Private/AnimNodes/AnimNode_PoseHandler.cpp b/Engine/Source/Runtime/AnimGraphRuntime/Private/AnimNodes/AnimNode_PoseHandler.cpp
--- a/Engine/Source/Runtime/AnimGraphRuntime/Private/AnimNodes/AnimNode_PoseHandler.cpp
+++ b/Engine/Source/Runtime/AnimGraphRuntime/Private/AnimNodes/AnimNode_PoseHandler.cpp
@@ -8,13 +8,15 @@
 
 void FAnimNode_PoseHandler::Initialize_AnyThread(const FAnimationInitializeContext& Context)
 {
+	DECLARE_SCOPE_HIERARCHICAL_COUNTER_ANIMNODE(Initialize_AnyThread)
 	FAnimNode_AssetPlayerBase::Initialize_AnyThread(Context);
 
 	UpdatePoseAssetProperty(Context.AnimInstanceProxy);
 }
 
-void FAnimNode_PoseHandler::CacheBones_AnyThread(const FAnimationCacheBonesContext& Context) 
+void FAnimNode_PoseHandler::CacheBones_AnyThread(const FAnimationCacheBonesContext& Context)
 {
+	DECLARE_SCOPE_HIERARCHICAL_COUNTER_ANIMNODE(CacheBones_AnyThread)
 	FAnimNode_AssetPlayerBase::CacheBones_AnyThread(Context);
 
 	BoneBlendWeights.Reset();
@@ -87,6 +89,7 @@ void FAnimNode_PoseHandler::OverrideAsset(UAnimationAsset* NewAsset)
 
 void FAnimNode_PoseHandler::GatherDebugData(FNodeDebugData& DebugData)
 {
+	DECLARE_SCOPE_HIERARCHICAL_COUNTER_ANIMNODE(GatherDebugData)
 	FString DebugLine = DebugData.GetNodeName(this);
 	
 	DebugLine += FString::Printf(TEXT("('%s')"), *GetNameSafe(PoseAsset));
diff --git a/Engine/Source/Runtime/AnimGraphRuntime/Private/AnimNodes/AnimNode_PoseSnapshot.cpp b/Engine/Source/Runtime/AnimGraphRuntime/Private/AnimNodes/AnimNode_PoseSnapshot.cpp
--- a/Engine/Source/Runtime/AnimGraphRuntime/Private/AnimNodes/AnimNode_PoseSnapshot.cpp
+++ b/Engine/Source/Runtime/AnimGraphRuntime/Private/AnimNodes/AnimNode_PoseSnapshot.cpp
@@ -48,12 +48,14 @@ void FAnimNode_PoseSnapshot::PreUpdate(const UAnimInstance* InAnimInstance)
 
 void FAnimNode_PoseSnapshot::Update_AnyThread(const FAnimationUpdateContext& Context)
 {
+	DECLARE_SCOPE_HIERARCHICAL_COUNTER_ANIMNODE(Update_AnyThread)
 	// Evaluate any BP logic plugged into this node
 	GetEvaluateGraphExposedInputs().Execute(Context);
 }
 
 void FAnimNode_PoseSnapshot::Evaluate_AnyThread(FPoseContext& Output)
 {
+	DECLARE_SCOPE_HIERARCHICAL_COUNTER_ANIMNODE(Evaluate_AnyThread)
 	FCompactPose& OutPose = Output.Pose;
 	OutPose.ResetToRefPose();
 
@@ -117,6 +119,7 @@ void FAnimNode_PoseSnapshot::ApplyPose(const FPoseSnapshot& PoseSnapshot, FCompa
 
 void FAnimNode_PoseSnapshot::GatherDebugData(FNodeDebugData& DebugData)
 {
+	DECLARE_SCOPE_HIERARCHICAL_COUNTER_ANIMNODE(GatherDebugData)
 	FString DebugLine = DebugData.GetNodeName(this) + " Snapshot Name:" + SnapshotName.ToString();
 	DebugData.AddDebugItem(DebugLine, true);
 }
diff --git a/Engine/Source/Runtime/AnimGraphRuntime/Private/AnimNodes/AnimNode_RandomPlayer.cpp b/Engine/Source/Runtime/AnimGraphRuntime/Private/AnimNodes/AnimNode_RandomPlayer.cpp
--- a/Engine/Source/Runtime/AnimGraphRuntime/Private/AnimNodes/AnimNode_RandomPlayer.cpp
+++ b/Engine/Source/Runtime/AnimGraphRuntime/Private/AnimNodes/AnimNode_RandomPlayer.cpp
@@ -16,6 +16,7 @@ FAnimNode_RandomPlayer::FAnimNode_RandomPlayer()
 
 void FAnimNode_RandomPlayer::Initialize_AnyThread(const FAnimationInitializeContext& Context)
 {
+	DECLARE_SCOPE_HIERARCHICAL_COUNTER_ANIMNODE(Initialize_AnyThread)
 	FAnimNode_Base::Initialize_AnyThread(Context);
 	GetEvaluateGraphExposedInputs().Execute(Context);
 
@@ -76,6 +77,7 @@ void FAnimNode_RandomPlayer::Initialize_AnyThread(const FAnimationInitializeCont
 
 void FAnimNode_RandomPlayer::Update_AnyThread(const FAnimationUpdateContext& Context)
 {
+	DECLARE_SCOPE_HIERARCHICAL_COUNTER_ANIMNODE(Update_AnyThread)
 	GetEvaluateGraphExposedInputs().Execute(Context);
 
 	if(Entries.Num() == 0)
@@ -168,6 +170,7 @@ void FAnimNode_RandomPlayer::Update_AnyThread(const FAnimationUpdateContext& Con
 
 void FAnimNode_RandomPlayer::Evaluate_AnyThread(FPoseContext& Output)
 {
+	DECLARE_SCOPE_HIERARCHICAL_COUNTER_ANIMNODE(Evaluate_AnyThread)
 	if(Entries.Num() > 0)
 	{
 		UAnimSequence* CurrentSequence = Entries[CurrentEntry].Sequence;
@@ -227,6 +230,7 @@ void FAnimNode_RandomPlayer::Evaluate_AnyThread(FPoseContext& Output)
 
 void FAnimNode_RandomPlayer::GatherDebugData(FNodeDebugData& DebugData)
 {
+	DECLARE_SCOPE_HIERARCHICAL_COUNTER_ANIMNODE(GatherDebugData)
 	FString DebugLine = DebugData.GetNodeName(this);
 
 	DebugData.AddDebugItem(DebugLine, true);
diff --git a/Engine/Source/Runtime/AnimGraphRuntime/Private/AnimNodes/AnimNode_RefPose.cpp b/Engine/Source/Runtime/AnimGraphRuntime/Private/AnimNodes/AnimNode_RefPose.cpp
--- a/Engine/Source/Runtime/AnimGraphRuntime/Private/AnimNodes/AnimNode_RefPose.cpp
+++ b/Engine/Source/Runtime/AnimGraphRuntime/Private/AnimNodes/AnimNode_RefPose.cpp
@@ -7,6 +7,7 @@
 
 void FAnimNode_RefPose::Evaluate_AnyThread(FPoseContext& Output)
 {
+	DECLARE_SCOPE_HIERARCHICAL_COUNTER_ANIMNODE(Evaluate_AnyThread)
 	// I don't have anything to evaluate. Should this be even here?
 	// EvaluateGraphExposedInputs.Execute(Context);
 
@@ -25,6 +26,7 @@ void FAnimNode_RefPose::Evaluate_AnyThread(FPoseContext& Output)
 
 void FAnimNode_MeshSpaceRefPose::EvaluateComponentSpace_AnyThread(FComponentSpacePoseContext& Output)
 {
+	DECLARE_SCOPE_HIERARCHICAL_COUNTER_ANIMNODE(EvaluateComponentSpace_AnyThread)
 	Output.ResetToRefPose();
 }
 
@@ -44,6 +46,7 @@ const TCHAR* GetRefPostTypeText(ERefPoseType RefPose)
 
 void FAnimNode_RefPose::GatherDebugData(FNodeDebugData& DebugData)
 {
+	DECLARE_SCOPE_HIERARCHICAL_COUNTER_ANIMNODE(GatherDebugData)
 	FString DebugLine = DebugData.GetNodeName(this);
 	DebugLine += FString::Printf(TEXT("(Ref Pose Type: %s)"), GetRefPostTypeText(RefPoseType));
 	DebugData.AddDebugItem(DebugLine, true);
diff --git a/Engine/Source/Runtime/AnimGraphRuntime/Private/AnimNodes/AnimNode_RotateRootBone.cpp b/Engine/Source/Runtime/AnimGraphRuntime/Private/AnimNodes/AnimNode_RotateRootBone.cpp
--- a/Engine/Source/Runtime/AnimGraphRuntime/Private/AnimNodes/AnimNode_RotateRootBone.cpp
+++ b/Engine/Source/Runtime/AnimGraphRuntime/Private/AnimNodes/AnimNode_RotateRootBone.cpp
@@ -7,6 +7,7 @@
 
 void FAnimNode_RotateRootBone::Initialize_AnyThread(const FAnimationInitializeContext& Context)
 {
+	DECLARE_SCOPE_HIERARCHICAL_COUNTER_ANIMNODE(Initialize_AnyThread)
 	FAnimNode_Base::Initialize_AnyThread(Context);
 
 	BasePose.Initialize(Context);
@@ -15,13 +16,15 @@ void FAnimNode_RotateRootBone::Initialize_AnyThread(const FAnimationInitializeCo
 	YawScaleBiasClamp.Reinitialize();
 }
 
-void FAnimNode_RotateRootBone::CacheBones_AnyThread(const FAnimationCacheBonesContext& Context) 
+void FAnimNode_RotateRootBone::CacheBones_AnyThread(const FAnimationCacheBonesContext& Context)
 {
+	DECLARE_SCOPE_HIERARCHICAL_COUNTER_ANIMNODE(CacheBones_AnyThread)
 	BasePose.CacheBones(Context);
 }
 
 void FAnimNode_RotateRootBone::Update_AnyThread(const FAnimationUpdateContext& Context)
 {
+	DECLARE_SCOPE_HIERARCHICAL_COUNTER_ANIMNODE(Update_AnyThread)
 	GetEvaluateGraphExposedInputs().Execute(Context);
 	BasePose.Update(Context);
 
@@ -31,6 +34,7 @@ void FAnimNode_RotateRootBone::Update_AnyThread(const FAnimationUpdateContext& C
 
 void FAnimNode_RotateRootBone::Evaluate_AnyThread(FPoseContext& Output)
 {
+	DECLARE_SCOPE_HIERARCHICAL_COUNTER_ANIMNODE(Evaluate_AnyThread)
 	// Evaluate the input
 	BasePose.Evaluate(Output);
 
@@ -57,6 +61,7 @@ void FAnimNode_RotateRootBone::Evaluate_AnyThread(FPoseContext& Output)
 
 void FAnimNode_RotateRootBone::GatherDebugData(FNodeDebugData& DebugData)
 {
+	DECLARE_SCOPE_HIERARCHICAL_COUNTER_ANIMNODE(GatherDebugData)
 	FString DebugLine = DebugData.GetNodeName(this);
 
 	DebugLine += FString::Printf(TEXT("Pitch(%.2f) Yaw(%.2f)"), ActualPitch, ActualYaw);
diff --git a/Engine/Source/Runtime/AnimGraphRuntime/Private/AnimNodes/AnimNode_RotationOffsetBlendSpace.cpp b/Engine/Source/Runtime/AnimGraphRuntime/Private/AnimNodes/AnimNode_RotationOffsetBlendSpace.cpp
--- a/Engine/Source/Runtime/AnimGraphRuntime/Private/AnimNodes/AnimNode_RotationOffsetBlendSpace.cpp
+++ b/Engine/Source/Runtime/AnimGraphRuntime/Private/AnimNodes/AnimNode_RotationOffsetBlendSpace.cpp
@@ -20,12 +20,14 @@ FAnimNode_RotationOffsetBlendSpace::FAnimNode_RotationOffsetBlendSpace()
 
 void FAnimNode_RotationOffsetBlendSpace::Initialize_AnyThread(const FAnimationInitializeContext& Context)
 {
+	DECLARE_SCOPE_HIERARCHICAL_COUNTER_ANIMNODE(Initialize_AnyThread)
 	FAnimNode_BlendSpacePlayer::Initialize_AnyThread(Context);
 	BasePose.Initialize(Context);
 }
 
-void FAnimNode_RotationOffsetBlendSpace::CacheBones_AnyThread(const FAnimationCacheBonesContext& Context) 
+void FAnimNode_RotationOffsetBlendSpace::CacheBones_AnyThread(const FAnimationCacheBonesContext& Context)
 {
+	DECLARE_SCOPE_HIERARCHICAL_COUNTER_ANIMNODE(CacheBones_AnyThread)
 	FAnimNode_BlendSpacePlayer::CacheBones_AnyThread(Context);
 	BasePose.CacheBones(Context);
 }
@@ -88,6 +90,7 @@ void FAnimNode_RotationOffsetBlendSpace::Evaluate_AnyThread(FPoseContext& Contex
 
 void FAnimNode_RotationOffsetBlendSpace::GatherDebugData(FNodeDebugData& DebugData)
 {
+	DECLARE_SCOPE_HIERARCHICAL_COUNTER_ANIMNODE(GatherDebugData)
 	FString DebugLine = DebugData.GetNodeName(this);
 	
 	DebugLine += FString::Printf(TEXT("Alpha (%.1f%%) PlayTime (%.3f)"), ActualAlpha * 100.f, InternalTimeAccumulator);
diff --git a/Engine/Source/Runtime/AnimGraphRuntime/Private/AnimNodes/AnimNode_SequenceEvaluator.cpp b/Engine/Source/Runtime/AnimGraphRuntime/Private/AnimNodes/AnimNode_SequenceEvaluator.cpp
--- a/Engine/Source/Runtime/AnimGraphRuntime/Private/AnimNodes/AnimNode_SequenceEvaluator.cpp
+++ b/Engine/Source/Runtime/AnimGraphRuntime/Private/AnimNodes/AnimNode_SequenceEvaluator.cpp
@@ -18,12 +18,14 @@ float FAnimNode_SequenceEvaluator::GetCurrentAssetLength()
 
 void FAnimNode_SequenceEvaluator::Initialize_AnyThread(const FAnimationInitializeContext& Context)
 {
+	DECLARE_SCOPE_HIERARCHICAL_COUNTER_ANIMNODE(Initialize_AnyThread)
 	FAnimNode_AssetPlayerBase::Initialize_AnyThread(Context);
 	bReinitialized = true;
 }
 
-void FAnimNode_SequenceEvaluator::CacheBones_AnyThread(const FAnimationCacheBonesContext& Context) 
+void FAnimNode_SequenceEvaluator::CacheBones_AnyThread(const FAnimationCacheBonesContext& Context)
 {
+	DECLARE_SCOPE_HIERARCHICAL_COUNTER_ANIMNODE(CacheBones_AnyThread)
 }
 
 void FAnimNode_SequenceEvaluator::UpdateAssetPlayer(const FAnimationUpdateContext& Context)
@@ -87,6 +89,7 @@ void FAnimNode_SequenceEvaluator::UpdateAssetPlayer(const FAnimationUpdateContex
 
 void FAnimNode_SequenceEvaluator::Evaluate_AnyThread(FPoseContext& Output)
 {
+	DECLARE_SCOPE_HIERARCHICAL_COUNTER_ANIMNODE(Evaluate_AnyThread)
 	check(Output.AnimInstanceProxy != nullptr);
 	if ((Sequence != nullptr) && (Output.AnimInstanceProxy->IsSkeletonCompatible(Sequence->GetSkeleton())))
 	{
@@ -108,6 +111,7 @@ void FAnimNode_SequenceEvaluator::OverrideAsset(UAnimationAsset* NewAsset)
 
 void FAnimNode_SequenceEvaluator::GatherDebugData(FNodeDebugData& DebugData)
 {
+	DECLARE_SCOPE_HIERARCHICAL_COUNTER_ANIMNODE(GatherDebugData)
 	FString DebugLine = DebugData.GetNodeName(this);
 	
 	DebugLine += FString::Printf(TEXT("('%s' InputTime: %.3f, Time: %.3f)"), *GetNameSafe(Sequence), ExplicitTime, InternalTimeAccumulator);
diff --git a/Engine/Source/Runtime/AnimGraphRuntime/Private/AnimNodes/AnimNode_Slot.cpp b/Engine/Source/Runtime/AnimGraphRuntime/Private/AnimNodes/AnimNode_Slot.cpp
--- a/Engine/Source/Runtime/AnimGraphRuntime/Private/AnimNodes/AnimNode_Slot.cpp
+++ b/Engine/Source/Runtime/AnimGraphRuntime/Private/AnimNodes/AnimNode_Slot.cpp
@@ -8,6 +8,7 @@
 
 void FAnimNode_Slot::Initialize_AnyThread(const FAnimationInitializeContext& Context)
 {
+	DECLARE_SCOPE_HIERARCHICAL_COUNTER_ANIMNODE(Initialize_AnyThread)
 	FAnimNode_Base::Initialize_AnyThread(Context);
 
 	Source.Initialize(Context);
@@ -21,13 +22,15 @@ void FAnimNode_Slot::Initialize_AnyThread(const FAnimationInitializeContext& Con
 	}
 }
 
-void FAnimNode_Slot::CacheBones_AnyThread(const FAnimationCacheBonesContext& Context) 
+void FAnimNode_Slot::CacheBones_AnyThread(const FAnimationCacheBonesContext& Context)
 {
+	DECLARE_SCOPE_HIERARCHICAL_COUNTER_ANIMNODE(CacheBones_AnyThread)
 	Source.CacheBones(Context);
 }
 
 void FAnimNode_Slot::Update_AnyThread(const FAnimationUpdateContext& Context)
 {
+	DECLARE_SCOPE_HIERARCHICAL_COUNTER_ANIMNODE(Update_AnyThread)
 	// Update weights.
 	Context.AnimInstanceProxy->GetSlotWeight(SlotName, WeightData.SlotNodeWeight, WeightData.SourceWeight, WeightData.TotalNodeWeight);
 
@@ -44,6 +47,7 @@ void FAnimNode_Slot::Update_AnyThread(const FAnimationUpdateContext& Context)
 
 void FAnimNode_Slot::Evaluate_AnyThread(FPoseContext & Output)
 {
+	DECLARE_SCOPE_HIERARCHICAL_COUNTER_ANIMNODE(Evaluate_AnyThread)
 	// If not playing a montage, just pass through
 	if (WeightData.SlotNodeWeight <= ZERO_ANIMWEIGHT_THRESH)
 	{
@@ -66,6 +70,7 @@ void FAnimNode_Slot::Evaluate_AnyThread(FPoseContext & Output)
 
 void FAnimNode_Slot::GatherDebugData(FNodeDebugData& DebugData)
 {
+	DECLARE_SCOPE_HIERARCHICAL_COUNTER_ANIMNODE(GatherDebugData)
 	FString DebugLine = DebugData.GetNodeName(this);
 	DebugLine += FString::Printf(TEXT("(Slot Name: '%s' Weight:%.1f%%)"), *SlotName.ToString(), WeightData.SlotNodeWeight*100.f);
 	
diff --git a/Engine/Source/Runtime/AnimGraphRuntime/Private/AnimNodes/AnimNode_TwoWayBlend.cpp b/Engine/Source/Runtime/AnimGraphRuntime/Private/AnimNodes/AnimNode_TwoWayBlend.cpp
--- a/Engine/Source/Runtime/AnimGraphRuntime/Private/AnimNodes/AnimNode_TwoWayBlend.cpp
+++ b/Engine/Source/Runtime/AnimGraphRuntime/Private/AnimNodes/AnimNode_TwoWayBlend.cpp
@@ -9,6 +9,7 @@
 
 void FAnimNode_TwoWayBlend::Initialize_AnyThread(const FAnimationInitializeContext& Context)
 {
+	DECLARE_SCOPE_HIERARCHICAL_COUNTER_ANIMNODE(Initialize_AnyThread)
 	FAnimNode_Base::Initialize_AnyThread(Context);
 
 	A.Initialize(Context);
@@ -21,14 +22,16 @@ void FAnimNode_TwoWayBlend::Initialize_AnyThread(const FAnimationInitializeConte
 	AlphaScaleBiasClamp.Reinitialize();
 }
 
-void FAnimNode_TwoWayBlend::CacheBones_AnyThread(const FAnimationCacheBonesContext& Context) 
+void FAnimNode_TwoWayBlend::CacheBones_AnyThread(const FAnimationCacheBonesContext& Context)
 {
+	DECLARE_SCOPE_HIERARCHICAL_COUNTER_ANIMNODE(CacheBones_AnyThread)
 	A.CacheBones(Context);
 	B.CacheBones(Context);
 }
 
 void FAnimNode_TwoWayBlend::Update_AnyThread(const FAnimationUpdateContext& Context)
 {
+	DECLARE_SCOPE_HIERARCHICAL_COUNTER_ANIMNODE(Update_AnyThread)
 	QUICK_SCOPE_CYCLE_COUNTER(STAT_FAnimationNode_TwoWayBlend_Update);
 	GetEvaluateGraphExposedInputs().Execute(Context);
 
@@ -101,6 +104,7 @@ void FAnimNode_TwoWayBlend::Update_AnyThread(const FAnimationUpdateContext& Cont
 
 void FAnimNode_TwoWayBlend::Evaluate_AnyThread(FPoseContext& Output)
 {
+	DECLARE_SCOPE_HIERARCHICAL_COUNTER_ANIMNODE(Evaluate_AnyThread)
 	if (bBIsRelevant)
 	{
 		if (bAIsRelevant)
@@ -127,6 +131,7 @@ void FAnimNode_TwoWayBlend::Evaluate_AnyThread(FPoseContext& Output)
 
 void FAnimNode_TwoWayBlend::GatherDebugData(FNodeDebugData& DebugData)
 {
+	DECLARE_SCOPE_HIERARCHICAL_COUNTER_ANIMNODE(GatherDebugData)
 	FString DebugLine = DebugData.GetNodeName(this);
 	DebugLine += FString::Printf(TEXT("(Alpha: %.1f%%)"), InternalBlendAlpha *100);
 	DebugData.AddDebugItem(DebugLine);
diff --git a/Engine/Source/Runtime/AnimGraphRuntime/Private/BoneControllers/AnimNodeRotationMultiplier.cpp b/Engine/Source/Runtime/AnimGraphRuntime/Private/BoneControllers/AnimNodeRotationMultiplier.cpp
--- a/Engine/Source/Runtime/AnimGraphRuntime/Private/BoneControllers/AnimNodeRotationMultiplier.cpp
+++ b/Engine/Source/Runtime/AnimGraphRuntime/Private/BoneControllers/AnimNodeRotationMultiplier.cpp
@@ -17,6 +17,7 @@ FAnimNode_RotationMultiplier::FAnimNode_RotationMultiplier()
 
 void FAnimNode_RotationMultiplier::GatherDebugData(FNodeDebugData& DebugData)
 {
+	DECLARE_SCOPE_HIERARCHICAL_COUNTER_ANIMNODE(GatherDebugData)
 	FString DebugLine = DebugData.GetNodeName(this);
 
 	DebugLine += "(";
@@ -126,6 +127,7 @@ FQuat FAnimNode_RotationMultiplier::MultiplyQuatBasedOnSourceIndex(const FTransf
 
 void FAnimNode_RotationMultiplier::EvaluateSkeletalControl_AnyThread(FComponentSpacePoseContext& Output, TArray<FBoneTransform>& OutBoneTransforms)
 {
+	DECLARE_SCOPE_HIERARCHICAL_COUNTER_ANIMNODE(EvaluateSkeletalControl_AnyThread)
 	check(OutBoneTransforms.Num() == 0);
 
 	if ( Multiplier != 0.f )
@@ -170,6 +172,7 @@ bool FAnimNode_RotationMultiplier::IsValidToEvaluate(const USkeleton* Skeleton,
 
 void FAnimNode_RotationMultiplier::InitializeBoneReferences(const FBoneContainer& RequiredBones) 
 {
+	DECLARE_SCOPE_HIERARCHICAL_COUNTER_ANIMNODE(InitializeBoneReferences)
 	SourceBone.Initialize(RequiredBones);
 	TargetBone.Initialize(RequiredBones);
 }
diff --git a/Engine/Source/Runtime/AnimGraphRuntime/Private/BoneControllers/AnimNode_AnimDynamics.cpp b/Engine/Source/Runtime/AnimGraphRuntime/Private/BoneControllers/AnimNode_AnimDynamics.cpp
--- a/Engine/Source/Runtime/AnimGraphRuntime/Private/BoneControllers/AnimNode_AnimDynamics.cpp
+++ b/Engine/Source/Runtime/AnimGraphRuntime/Private/BoneControllers/AnimNode_AnimDynamics.cpp
@@ -167,6 +167,7 @@ FAnimNode_AnimDynamics::FAnimNode_AnimDynamics()
 
 void FAnimNode_AnimDynamics::Initialize_AnyThread(const FAnimationInitializeContext& Context)
 {
+	DECLARE_SCOPE_HIERARCHICAL_COUNTER_ANIMNODE(Initialize_AnyThread)
 	FAnimNode_SkeletalControlBase::Initialize_AnyThread(Context);
 
 	FBoneContainer& RequiredBones = Context.AnimInstanceProxy->GetRequiredBones();
@@ -187,6 +188,7 @@ void FAnimNode_AnimDynamics::Initialize_AnyThread(const FAnimationInitializeCont
 
 void FAnimNode_AnimDynamics::UpdateInternal(const FAnimationUpdateContext& Context)
 {
+	DECLARE_SCOPE_HIERARCHICAL_COUNTER_ANIMNODE(UpdateInternal)
 	FAnimNode_SkeletalControlBase::UpdateInternal(Context);
 
 	NextTimeStep = Context.GetDeltaTime();
@@ -208,6 +210,7 @@ bool FAnimNode_AnimDynamics::IsAnimDynamicsSystemEnabledFor(int32 InLOD)
 }
 void FAnimNode_AnimDynamics::EvaluateSkeletalControl_AnyThread(FComponentSpacePoseContext& Output, TArray<FBoneTransform>& OutBoneTransforms)
 {
+	DECLARE_SCOPE_HIERARCHICAL_COUNTER_ANIMNODE(EvaluateSkeletalControl_AnyThread)
 	SCOPE_CYCLE_COUNTER(STAT_AnimDynamicsOverall);
 
 	if (IsAnimDynamicsSystemEnabledFor(Output.AnimInstanceProxy->GetLODLevel()))
@@ -397,6 +400,7 @@ void FAnimNode_AnimDynamics::EvaluateSkeletalControl_AnyThread(FComponentSpacePo
 
 void FAnimNode_AnimDynamics::InitializeBoneReferences(const FBoneContainer& RequiredBones)
 {
+	DECLARE_SCOPE_HIERARCHICAL_COUNTER_ANIMNODE(InitializeBoneReferences)
 	BoundBone.Initialize(RequiredBones);
 
 	if (bChain)
@@ -456,6 +460,7 @@ void FAnimNode_AnimDynamics::InitializeBoneReferences(const FBoneContainer& Requ
 
 void FAnimNode_AnimDynamics::GatherDebugData(FNodeDebugData& DebugData)
 {
+	DECLARE_SCOPE_HIERARCHICAL_COUNTER_ANIMNODE(GatherDebugData)
 	const float ActualBiasedAlpha = AlphaScaleBias.ApplyTo(Alpha);
 
 	FString DebugLine = DebugData.GetNodeName(this);
diff --git a/Engine/Source/Runtime/AnimGraphRuntime/Private/BoneControllers/AnimNode_ApplyLimits.cpp b/Engine/Source/Runtime/AnimGraphRuntime/Private/BoneControllers/AnimNode_ApplyLimits.cpp
--- a/Engine/Source/Runtime/AnimGraphRuntime/Private/BoneControllers/AnimNode_ApplyLimits.cpp
+++ b/Engine/Source/Runtime/AnimGraphRuntime/Private/BoneControllers/AnimNode_ApplyLimits.cpp
@@ -15,6 +15,7 @@ FAnimNode_ApplyLimits::FAnimNode_ApplyLimits()
 
 void FAnimNode_ApplyLimits::GatherDebugData(FNodeDebugData& DebugData)
 {
+	DECLARE_SCOPE_HIERARCHICAL_COUNTER_ANIMNODE(GatherDebugData)
 	FString DebugLine = DebugData.GetNodeName(this);
 
 	DebugLine += "(";
@@ -28,6 +29,7 @@ void FAnimNode_ApplyLimits::GatherDebugData(FNodeDebugData& DebugData)
 
 void FAnimNode_ApplyLimits::EvaluateSkeletalControl_AnyThread(FComponentSpacePoseContext& Output, TArray<FBoneTransform>& OutBoneTransforms)
 {
+	DECLARE_SCOPE_HIERARCHICAL_COUNTER_ANIMNODE(EvaluateSkeletalControl_AnyThread)
 	checkSlow(OutBoneTransforms.Num() == 0);
 
 	FPoseContext LocalPose0(Output.AnimInstanceProxy);
@@ -96,6 +98,7 @@ void FAnimNode_ApplyLimits::OnInitializeAnimInstance(const FAnimInstanceProxy* I
 
 void FAnimNode_ApplyLimits::InitializeBoneReferences(const FBoneContainer& RequiredBones)
 {
+	DECLARE_SCOPE_HIERARCHICAL_COUNTER_ANIMNODE(InitializeBoneReferences)
 	for (FAngularRangeLimit& AngularLimit : AngularRangeLimits)
 	{
 		AngularLimit.Bone.Initialize(RequiredBones);
diff --git a/Engine/Source/Runtime/AnimGraphRuntime/Private/BoneControllers/AnimNode_BoneDrivenController.cpp b/Engine/Source/Runtime/AnimGraphRuntime/Private/BoneControllers/AnimNode_BoneDrivenController.cpp
--- a/Engine/Source/Runtime/AnimGraphRuntime/Private/BoneControllers/AnimNode_BoneDrivenController.cpp
+++ b/Engine/Source/Runtime/AnimGraphRuntime/Private/BoneControllers/AnimNode_BoneDrivenController.cpp
@@ -36,6 +36,7 @@ FAnimNode_BoneDrivenController::FAnimNode_BoneDrivenController()
 
 void FAnimNode_BoneDrivenController::GatherDebugData(FNodeDebugData& DebugData)
 {
+	DECLARE_SCOPE_HIERARCHICAL_COUNTER_ANIMNODE(GatherDebugData)
 	FString DebugLine = DebugData.GetNodeName(this);
 	DebugLine += "(";
 	AddDebugNodeData(DebugLine);
@@ -55,6 +56,7 @@ void FAnimNode_BoneDrivenController::GatherDebugData(FNodeDebugData& DebugData)
 
 void FAnimNode_BoneDrivenController::EvaluateSkeletalControl_AnyThread(FComponentSpacePoseContext& Output, TArray<FBoneTransform>& OutBoneTransforms)
 {
+	DECLARE_SCOPE_HIERARCHICAL_COUNTER_ANIMNODE(EvaluateSkeletalControl_AnyThread)
 	check(OutBoneTransforms.Num() == 0);
 
 	// Early out if we're not driving from or to anything
@@ -171,6 +173,7 @@ void FAnimNode_BoneDrivenController::EvaluateSkeletalControl_AnyThread(FComponen
 
 void FAnimNode_BoneDrivenController::EvaluateComponentSpaceInternal(FComponentSpacePoseContext& Context)
 {
+	DECLARE_SCOPE_HIERARCHICAL_COUNTER_ANIMNODE(EvaluateComponentSpaceInternal)
 	// Early out if we're not driving from or to anything
 	if (SourceComponent == EComponentType::None || DestinationMode == EDrivenDestinationMode::Bone)
 	{
@@ -286,6 +289,7 @@ void FAnimNode_BoneDrivenController::ConvertTargetComponentToBits()
 
 void FAnimNode_BoneDrivenController::InitializeBoneReferences(const FBoneContainer& RequiredBones)
 {
+	DECLARE_SCOPE_HIERARCHICAL_COUNTER_ANIMNODE(InitializeBoneReferences)
 	SourceBone.Initialize(RequiredBones);
 	TargetBone.Initialize(RequiredBones);
 }
diff --git a/Engine/Source/Runtime/AnimGraphRuntime/Private/BoneControllers/AnimNode_CCDIK.cpp b/Engine/Source/Runtime/AnimGraphRuntime/Private/BoneControllers/AnimNode_CCDIK.cpp
--- a/Engine/Source/Runtime/AnimGraphRuntime/Private/BoneControllers/AnimNode_CCDIK.cpp
+++ b/Engine/Source/Runtime/AnimGraphRuntime/Private/BoneControllers/AnimNode_CCDIK.cpp
@@ -46,6 +46,7 @@ FTransform FAnimNode_CCDIK::GetTargetTransform(const FTransform& InComponentTran
 
 void FAnimNode_CCDIK::EvaluateSkeletalControl_AnyThread(FComponentSpacePoseContext& Output, TArray<FBoneTransform>& OutBoneTransforms)
 {
+	DECLARE_SCOPE_HIERARCHICAL_COUNTER_ANIMNODE(EvaluateSkeletalControl_AnyThread)
 	const FBoneContainer& BoneContainer = Output.Pose.GetPose().GetBoneContainer();
 
 	// Update EffectorLocation if it is based off a bone position
@@ -185,6 +186,7 @@ void FAnimNode_CCDIK::ResizeRotationLimitPerJoints(int32 NewSize)
 
 void FAnimNode_CCDIK::InitializeBoneReferences(const FBoneContainer& RequiredBones)
 {
+	DECLARE_SCOPE_HIERARCHICAL_COUNTER_ANIMNODE(InitializeBoneReferences)
 	TipBone.Initialize(RequiredBones);
 	RootBone.Initialize(RequiredBones);
 	EffectorTarget.InitializeBoneReferences(RequiredBones);
@@ -192,6 +194,7 @@ void FAnimNode_CCDIK::InitializeBoneReferences(const FBoneContainer& RequiredBon
 
 void FAnimNode_CCDIK::GatherDebugData(FNodeDebugData& DebugData)
 {
+	DECLARE_SCOPE_HIERARCHICAL_COUNTER_ANIMNODE(GatherDebugData)
 	FString DebugLine = DebugData.GetNodeName(this);
 
 	DebugData.AddDebugItem(DebugLine);
diff --git a/Engine/Source/Runtime/AnimGraphRuntime/Private/BoneControllers/AnimNode_Fabrik.cpp b/Engine/Source/Runtime/AnimGraphRuntime/Private/BoneControllers/AnimNode_Fabrik.cpp
--- a/Engine/Source/Runtime/AnimGraphRuntime/Private/BoneControllers/AnimNode_Fabrik.cpp
+++ b/Engine/Source/Runtime/AnimGraphRuntime/Private/BoneControllers/AnimNode_Fabrik.cpp
@@ -49,6 +49,7 @@ FTransform FAnimNode_Fabrik::GetTargetTransform(const FTransform& InComponentTra
 
 void FAnimNode_Fabrik::EvaluateSkeletalControl_AnyThread(FComponentSpacePoseContext& Output, TArray<FBoneTransform>& OutBoneTransforms)
 {
+	DECLARE_SCOPE_HIERARCHICAL_COUNTER_ANIMNODE(EvaluateSkeletalControl_AnyThread)
 	const FBoneContainer& BoneContainer = Output.Pose.GetPose().GetBoneContainer();
 
 	// Update EffectorLocation if it is based off a bone position
@@ -223,6 +224,7 @@ void FAnimNode_Fabrik::ConditionalDebugDraw(FPrimitiveDrawInterface* PDI, USkele
 
 void FAnimNode_Fabrik::InitializeBoneReferences(const FBoneContainer& RequiredBones)
 {
+	DECLARE_SCOPE_HIERARCHICAL_COUNTER_ANIMNODE(InitializeBoneReferences)
 	TipBone.Initialize(RequiredBones);
 	RootBone.Initialize(RequiredBones);
 	EffectorTarget.InitializeBoneReferences(RequiredBones);
@@ -230,6 +232,7 @@ void FAnimNode_Fabrik::InitializeBoneReferences(const FBoneContainer& RequiredBo
 
 void FAnimNode_Fabrik::GatherDebugData(FNodeDebugData& DebugData)
 {
+	DECLARE_SCOPE_HIERARCHICAL_COUNTER_ANIMNODE(GatherDebugData)
 	FString DebugLine = DebugData.GetNodeName(this);
 
 	DebugData.AddDebugItem(DebugLine);
@@ -238,6 +241,7 @@ void FAnimNode_Fabrik::GatherDebugData(FNodeDebugData& DebugData)
 
 void FAnimNode_Fabrik::Initialize_AnyThread(const FAnimationInitializeContext& Context)
 {
+	DECLARE_SCOPE_HIERARCHICAL_COUNTER_ANIMNODE(Initialize_AnyThread)
 	Super::Initialize_AnyThread(Context);
 	EffectorTarget.Initialize(Context.AnimInstanceProxy);
 }
diff --git a/Engine/Source/Runtime/AnimGraphRuntime/Private/BoneControllers/AnimNode_LegIK.cpp b/Engine/Source/Runtime/AnimGraphRuntime/Private/BoneControllers/AnimNode_LegIK.cpp
--- a/Engine/Source/Runtime/AnimGraphRuntime/Private/BoneControllers/AnimNode_LegIK.cpp
+++ b/Engine/Source/Runtime/AnimGraphRuntime/Private/BoneControllers/AnimNode_LegIK.cpp
@@ -31,6 +31,7 @@ FAnimNode_LegIK::FAnimNode_LegIK()
 
 void FAnimNode_LegIK::GatherDebugData(FNodeDebugData& DebugData)
 {
+	DECLARE_SCOPE_HIERARCHICAL_COUNTER_ANIMNODE(GatherDebugData)
 	FString DebugLine = DebugData.GetNodeName(this);
 
 	// 	DebugLine += "(";
@@ -62,6 +63,7 @@ static void DrawDebugLeg(const FAnimLegIKData& InLegData, FAnimInstanceProxy* My
 
 void FAnimNode_LegIK::Initialize_AnyThread(const FAnimationInitializeContext& Context)
 {
+	DECLARE_SCOPE_HIERARCHICAL_COUNTER_ANIMNODE(Initialize_AnyThread)
 	FAnimNode_SkeletalControlBase::Initialize_AnyThread(Context);
 
 	MyAnimInstanceProxy = Context.AnimInstanceProxy;
@@ -90,6 +92,7 @@ void FAnimLegIKData::InitializeTransforms(FAnimInstanceProxy* MyAnimInstanceProx
 
 void FAnimNode_LegIK::EvaluateSkeletalControl_AnyThread(FComponentSpacePoseContext& Output, TArray<FBoneTransform>& OutBoneTransforms)
 {
+	DECLARE_SCOPE_HIERARCHICAL_COUNTER_ANIMNODE(EvaluateSkeletalControl_AnyThread)
 	SCOPE_CYCLE_COUNTER(STAT_LegIK_Eval);
 
 	check(OutBoneTransforms.Num() == 0);
@@ -799,6 +802,7 @@ static void PopulateLegBoneIndices(FAnimLegIKData& InLegData, const FCompactPose
 
 void FAnimNode_LegIK::InitializeBoneReferences(const FBoneContainer& RequiredBones)
 {
+	DECLARE_SCOPE_HIERARCHICAL_COUNTER_ANIMNODE(InitializeBoneReferences)
 	// Preserve FIKChain for each leg, as we're trying to maintain CachedBendDir between LOD transitions.
 	TMap<FName, FIKChain> IKChainLUT;
 	for(const FAnimLegIKData& LegData : LegsData)
diff --git a/Engine/Source/Runtime/AnimGraphRuntime/Private/BoneControllers/AnimNode_LookAt.cpp b/Engine/Source/Runtime/AnimGraphRuntime/Private/BoneControllers/AnimNode_LookAt.cpp
--- a/Engine/Source/Runtime/AnimGraphRuntime/Private/BoneControllers/AnimNode_LookAt.cpp
+++ b/Engine/Source/Runtime/AnimGraphRuntime/Private/BoneControllers/AnimNode_LookAt.cpp
@@ -39,6 +39,7 @@ FAnimNode_LookAt::FAnimNode_LookAt()
 
 void FAnimNode_LookAt::GatherDebugData(FNodeDebugData& DebugData)
 {
+	DECLARE_SCOPE_HIERARCHICAL_COUNTER_ANIMNODE(GatherDebugData)
 	FString DebugLine = DebugData.GetNodeName(this);
 
 	DebugLine += "(";
@@ -74,6 +75,7 @@ float FAnimNode_LookAt::AlphaToBlendType(float InAlpha, EInterpolationBlend::Typ
 
 void FAnimNode_LookAt::EvaluateSkeletalControl_AnyThread(FComponentSpacePoseContext& Output, TArray<FBoneTransform>& OutBoneTransforms)
 {
+	DECLARE_SCOPE_HIERARCHICAL_COUNTER_ANIMNODE(EvaluateSkeletalControl_AnyThread)
 	check(OutBoneTransforms.Num() == 0);
 
 	const FBoneContainer& BoneContainer = Output.Pose.GetPose().GetBoneContainer();
@@ -144,6 +146,7 @@ void FAnimNode_LookAt::EvaluateSkeletalControl_AnyThread(FComponentSpacePoseCont
 
 void FAnimNode_LookAt::EvaluateComponentSpaceInternal(FComponentSpacePoseContext& Context)
 {
+	DECLARE_SCOPE_HIERARCHICAL_COUNTER_ANIMNODE(EvaluateComponentSpaceInternal)
 	Super::EvaluateComponentSpaceInternal(Context);
 
 #if !(UE_BUILD_SHIPPING || UE_BUILD_TEST)
@@ -236,12 +239,14 @@ void FAnimNode_LookAt::ConditionalDebugDraw(FPrimitiveDrawInterface* PDI, USkele
 
 void FAnimNode_LookAt::InitializeBoneReferences(const FBoneContainer& RequiredBones)
 {
+	DECLARE_SCOPE_HIERARCHICAL_COUNTER_ANIMNODE(InitializeBoneReferences)
 	BoneToModify.Initialize(RequiredBones);
 	LookAtTarget.InitializeBoneReferences(RequiredBones);
 }
 
 void FAnimNode_LookAt::UpdateInternal(const FAnimationUpdateContext& Context)
 {
+	DECLARE_SCOPE_HIERARCHICAL_COUNTER_ANIMNODE(UpdateInternal)
 	FAnimNode_SkeletalControlBase::UpdateInternal(Context);
 
 	AccumulatedInterpoolationTime = FMath::Clamp(AccumulatedInterpoolationTime+Context.GetDeltaTime(), 0.f, InterpolationTime);;
@@ -249,6 +254,7 @@ void FAnimNode_LookAt::UpdateInternal(const FAnimationUpdateContext& Context)
 
 void FAnimNode_LookAt::Initialize_AnyThread(const FAnimationInitializeContext& Context)
 {
+	DECLARE_SCOPE_HIERARCHICAL_COUNTER_ANIMNODE(Initialize_AnyThread)
 	FAnimNode_SkeletalControlBase::Initialize_AnyThread(Context);
 
 	LookAtTarget.Initialize(Context.AnimInstanceProxy);
diff --git a/Engine/Source/Runtime/AnimGraphRuntime/Private/BoneControllers/AnimNode_RigidBody.cpp b/Engine/Source/Runtime/AnimGraphRuntime/Private/BoneControllers/AnimNode_RigidBody.cpp
--- a/Engine/Source/Runtime/AnimGraphRuntime/Private/BoneControllers/AnimNode_RigidBody.cpp
+++ b/Engine/Source/Runtime/AnimGraphRuntime/Private/BoneControllers/AnimNode_RigidBody.cpp
@@ -65,6 +65,7 @@ FAnimNode_RigidBody::~FAnimNode_RigidBody()
 
 void FAnimNode_RigidBody::GatherDebugData(FNodeDebugData& DebugData)
 {
+	DECLARE_SCOPE_HIERARCHICAL_COUNTER_ANIMNODE(GatherDebugData)
 	FString DebugLine = DebugData.GetNodeName(this);
 
 	DebugLine += "(";
@@ -117,6 +118,7 @@ FORCEINLINE_DEBUGGABLE FTransform ConvertCSTransformToSimSpace(ESimulationSpace
 
 void FAnimNode_RigidBody::UpdateComponentPose_AnyThread(const FAnimationUpdateContext& Context)
 {
+	DECLARE_SCOPE_HIERARCHICAL_COUNTER_ANIMNODE(UpdateComponentPose_AnyThread)
 	// Only freeze update graph after initial update, as we want to get that pose through.
 	if (bFreezeIncomingPoseOnStart && bSimulationStarted && ResetSimulatedTeleportType == ETeleportType::None)
 	{
@@ -143,6 +145,7 @@ void FAnimNode_RigidBody::UpdateComponentPose_AnyThread(const FAnimationUpdateCo
 
 void FAnimNode_RigidBody::EvaluateComponentPose_AnyThread(FComponentSpacePoseContext& Output)
 {
+	DECLARE_SCOPE_HIERARCHICAL_COUNTER_ANIMNODE(EvaluateComponentPose_AnyThread)
 	if (bFreezeIncomingPoseOnStart && bSimulationStarted)
 	{
 		// If we have a Frozen Pose captured, use it.
@@ -217,6 +220,7 @@ DECLARE_CYCLE_STAT(TEXT("FAnimNode_Ragdoll::EvaluateSkeletalControl_AnyThread"),
 
 void FAnimNode_RigidBody::EvaluateSkeletalControl_AnyThread(FComponentSpacePoseContext& Output, TArray<FBoneTransform>& OutBoneTransforms)
 {
+	DECLARE_SCOPE_HIERARCHICAL_COUNTER_ANIMNODE(EvaluateSkeletalControl_AnyThread)
 	SCOPE_CYCLE_COUNTER(STAT_RigidBody_Eval);
 	FScopeCycleCounterUObject AdditionalScope(UsePhysicsAsset, GET_STATID(STAT_RigidBody_Eval));
 	SCOPE_CYCLE_COUNTER(STAT_ImmediateEvaluateSkeletalControl);
@@ -965,6 +969,7 @@ DECLARE_CYCLE_STAT(TEXT("RigidBody_Update"), STAT_RigidBody_Update, STATGROUP_An
 
 void FAnimNode_RigidBody::UpdateInternal(const FAnimationUpdateContext& Context)
 {
+	DECLARE_SCOPE_HIERARCHICAL_COUNTER_ANIMNODE(UpdateInternal)
 	// Avoid this work if RBN is disabled, as the results would be discarded
 	if(CVarEnableRigidBodyNode.GetValueOnAnyThread() == 0)
 	{
@@ -1015,6 +1020,7 @@ void FAnimNode_RigidBody::UpdateInternal(const FAnimationUpdateContext& Context)
 
 void FAnimNode_RigidBody::InitializeBoneReferences(const FBoneContainer& RequiredBones) 
 {
+	DECLARE_SCOPE_HIERARCHICAL_COUNTER_ANIMNODE(InitializeBoneReferences)
 	/** We only need to update simulated bones and children of simulated bones*/
 	const int32 NumBodies = Bodies.Num();
 	const TArray<FBoneIndexType>& RequiredBoneIndices = RequiredBones.GetBoneIndicesArray();
diff --git a/Engine/Source/Runtime/AnimGraphRuntime/Private/BoneControllers/AnimNode_ScaleChainLength.cpp b/Engine/Source/Runtime/AnimGraphRuntime/Private/BoneControllers/AnimNode_ScaleChainLength.cpp
--- a/Engine/Source/Runtime/AnimGraphRuntime/Private/BoneControllers/AnimNode_ScaleChainLength.cpp
+++ b/Engine/Source/Runtime/AnimGraphRuntime/Private/BoneControllers/AnimNode_ScaleChainLength.cpp
@@ -20,12 +20,14 @@ FAnimNode_ScaleChainLength::FAnimNode_ScaleChainLength()
 
 void FAnimNode_ScaleChainLength::Initialize_AnyThread(const FAnimationInitializeContext& Context)
 {
+	DECLARE_SCOPE_HIERARCHICAL_COUNTER_ANIMNODE(Initialize_AnyThread)
 	FAnimNode_Base::Initialize_AnyThread(Context);
 	InputPose.Initialize(Context);
 }
 
 void FAnimNode_ScaleChainLength::Update_AnyThread(const FAnimationUpdateContext& Context)
 {
+	DECLARE_SCOPE_HIERARCHICAL_COUNTER_ANIMNODE(Update_AnyThread)
 	GetEvaluateGraphExposedInputs().Execute(Context);
 	InputPose.Update(Context);
 
@@ -34,6 +36,7 @@ void FAnimNode_ScaleChainLength::Update_AnyThread(const FAnimationUpdateContext&
 
 void FAnimNode_ScaleChainLength::CacheBones_AnyThread(const FAnimationCacheBonesContext& Context)
 {
+	DECLARE_SCOPE_HIERARCHICAL_COUNTER_ANIMNODE(CacheBones_AnyThread)
 	InputPose.CacheBones(Context);
 
 	// LOD change, recache bone indices.
@@ -42,6 +45,7 @@ void FAnimNode_ScaleChainLength::CacheBones_AnyThread(const FAnimationCacheBones
 
 void FAnimNode_ScaleChainLength::Evaluate_AnyThread(FPoseContext& Output)
 {
+	DECLARE_SCOPE_HIERARCHICAL_COUNTER_ANIMNODE(Evaluate_AnyThread)
 	// Evaluate incoming pose into our output buffer.
 	InputPose.Evaluate(Output);
 
@@ -146,6 +150,7 @@ float FAnimNode_ScaleChainLength::GetInitialChainLength(FCompactPose& InLSPose,
 
 void FAnimNode_ScaleChainLength::GatherDebugData(FNodeDebugData& DebugData)
 {
+	DECLARE_SCOPE_HIERARCHICAL_COUNTER_ANIMNODE(GatherDebugData)
 	FString DebugLine = DebugData.GetNodeName(this);
 	DebugLine += FString::Printf(TEXT("Alpha (%.1f%%)"), ActualAlpha * 100.f);
 	DebugData.AddDebugItem(DebugLine);
diff --git a/Engine/Source/Runtime/AnimGraphRuntime/Private/BoneControllers/AnimNode_SkeletalControlBase.cpp b/Engine/Source/Runtime/AnimGraphRuntime/Private/BoneControllers/AnimNode_SkeletalControlBase.cpp
--- a/Engine/Source/Runtime/AnimGraphRuntime/Private/BoneControllers/AnimNode_SkeletalControlBase.cpp
+++ b/Engine/Source/Runtime/AnimGraphRuntime/Private/BoneControllers/AnimNode_SkeletalControlBase.cpp
@@ -50,6 +50,7 @@ void FSocketReference::InitialzeCompactBoneIndex(const FBoneContainer& RequiredB
 
 void FAnimNode_SkeletalControlBase::Initialize_AnyThread(const FAnimationInitializeContext& Context)
 {
+	DECLARE_SCOPE_HIERARCHICAL_COUNTER_ANIMNODE(Initialize_AnyThread)
 	FAnimNode_Base::Initialize_AnyThread(Context);
 
 	ComponentPose.Initialize(Context);
@@ -58,8 +59,9 @@ void FAnimNode_SkeletalControlBase::Initialize_AnyThread(const FAnimationInitial
 	AlphaScaleBiasClamp.Reinitialize();
 }
 
-void FAnimNode_SkeletalControlBase::CacheBones_AnyThread(const FAnimationCacheBonesContext& Context) 
+void FAnimNode_SkeletalControlBase::CacheBones_AnyThread(const FAnimationCacheBonesContext& Context)
 {
+	DECLARE_SCOPE_HIERARCHICAL_COUNTER_ANIMNODE(CacheBones_AnyThread)
 	FAnimNode_Base::CacheBones_AnyThread(Context);
 	InitializeBoneReferences(Context.AnimInstanceProxy->GetRequiredBones());
 	ComponentPose.CacheBones(Context);
@@ -67,15 +69,18 @@ void FAnimNode_SkeletalControlBase::CacheBones_AnyThread(const FAnimationCacheBo
 
 void FAnimNode_SkeletalControlBase::UpdateInternal(const FAnimationUpdateContext& Context)
 {
+	DECLARE_SCOPE_HIERARCHICAL_COUNTER_ANIMNODE(UpdateInternal)
 }
 
 void FAnimNode_SkeletalControlBase::UpdateComponentPose_AnyThread(const FAnimationUpdateContext& Context)
 {
+	DECLARE_SCOPE_HIERARCHICAL_COUNTER_ANIMNODE(UpdateComponentPose_AnyThread)
 	ComponentPose.Update(Context);
 }
 
 void FAnimNode_SkeletalControlBase::Update_AnyThread(const FAnimationUpdateContext& Context)
 {
+	DECLARE_SCOPE_HIERARCHICAL_COUNTER_ANIMNODE(Update_AnyThread)
 	UpdateComponentPose_AnyThread(Context);
 
 	ActualAlpha = 0.f;
@@ -125,6 +130,7 @@ bool ContainsNaN(const TArray<FBoneTransform> & BoneTransforms)
 
 void FAnimNode_SkeletalControlBase::EvaluateComponentPose_AnyThread(FComponentSpacePoseContext& Output)
 {
+	DECLARE_SCOPE_HIERARCHICAL_COUNTER_ANIMNODE(EvaluateComponentPose_AnyThread)
 	// Evaluate the input
 	ComponentPose.EvaluateComponentSpace(Output);
 }
@@ -135,6 +141,7 @@ void FAnimNode_SkeletalControlBase::EvaluateComponentSpaceInternal(FComponentSpa
 
 void FAnimNode_SkeletalControlBase::EvaluateComponentSpace_AnyThread(FComponentSpacePoseContext& Output)
 {
+	DECLARE_SCOPE_HIERARCHICAL_COUNTER_ANIMNODE(EvaluateComponentSpace_AnyThread)
 	EvaluateComponentPose_AnyThread(Output);
 
 #if WITH_EDITORONLY_DATA
@@ -172,6 +179,7 @@ void FAnimNode_SkeletalControlBase::AddDebugNodeData(FString& OutDebugData)
 
 void FAnimNode_SkeletalControlBase::EvaluateSkeletalControl_AnyThread(FComponentSpacePoseContext& Output, TArray<FBoneTransform>& OutBoneTransforms)
 {
+	DECLARE_SCOPE_HIERARCHICAL_COUNTER_ANIMNODE(EvaluateSkeletalControl_AnyThread)
 	PRAGMA_DISABLE_DEPRECATION_WARNINGS
 	// Call legacy implementation for backwards compatibility
 	EvaluateBoneTransforms(Output.AnimInstanceProxy->GetSkelMeshComponent(), Output.Pose, OutBoneTransforms);
diff --git a/Engine/Source/Runtime/AnimGraphRuntime/Private/BoneControllers/AnimNode_SpringBone.cpp b/Engine/Source/Runtime/AnimGraphRuntime/Private/BoneControllers/AnimNode_SpringBone.cpp
--- a/Engine/Source/Runtime/AnimGraphRuntime/Private/BoneControllers/AnimNode_SpringBone.cpp
+++ b/Engine/Source/Runtime/AnimGraphRuntime/Private/BoneControllers/AnimNode_SpringBone.cpp
@@ -32,18 +32,21 @@ FAnimNode_SpringBone::FAnimNode_SpringBone()
 
 void FAnimNode_SpringBone::Initialize_AnyThread(const FAnimationInitializeContext& Context)
 {
+	DECLARE_SCOPE_HIERARCHICAL_COUNTER_ANIMNODE(Initialize_AnyThread)
 	FAnimNode_SkeletalControlBase::Initialize_AnyThread(Context);
 
 	RemainingTime = 0.0f;
 }
 
-void FAnimNode_SpringBone::CacheBones_AnyThread(const FAnimationCacheBonesContext& Context) 
+void FAnimNode_SpringBone::CacheBones_AnyThread(const FAnimationCacheBonesContext& Context)
 {
+	DECLARE_SCOPE_HIERARCHICAL_COUNTER_ANIMNODE(CacheBones_AnyThread)
 	FAnimNode_SkeletalControlBase::CacheBones_AnyThread(Context);
 }
 
 void FAnimNode_SpringBone::UpdateInternal(const FAnimationUpdateContext& Context)
 {
+	DECLARE_SCOPE_HIERARCHICAL_COUNTER_ANIMNODE(UpdateInternal)
 	FAnimNode_SkeletalControlBase::UpdateInternal(Context);
 
 	RemainingTime += Context.GetDeltaTime();
@@ -54,6 +57,7 @@ void FAnimNode_SpringBone::UpdateInternal(const FAnimationUpdateContext& Context
 
 void FAnimNode_SpringBone::GatherDebugData(FNodeDebugData& DebugData)
 {
+	DECLARE_SCOPE_HIERARCHICAL_COUNTER_ANIMNODE(GatherDebugData)
 	const float ActualBiasedAlpha = AlphaScaleBias.ApplyTo(Alpha);
 
 	//MDW_TODO Add more output info?
@@ -82,6 +86,7 @@ FORCEINLINE void CopyToVectorByFlags(FVector& DestVec, const FVector& SrcVec, bo
 
 void FAnimNode_SpringBone::EvaluateSkeletalControl_AnyThread(FComponentSpacePoseContext& Output, TArray<FBoneTransform>& OutBoneTransforms)
 {
+	DECLARE_SCOPE_HIERARCHICAL_COUNTER_ANIMNODE(EvaluateSkeletalControl_AnyThread)
 	check(OutBoneTransforms.Num() == 0);
 
 	const bool bNoOffset = !bTranslateX && !bTranslateY && !bTranslateZ;
@@ -218,6 +223,7 @@ bool FAnimNode_SpringBone::IsValidToEvaluate(const USkeleton* Skeleton, const FB
 
 void FAnimNode_SpringBone::InitializeBoneReferences(const FBoneContainer& RequiredBones) 
 {
+	DECLARE_SCOPE_HIERARCHICAL_COUNTER_ANIMNODE(InitializeBoneReferences)
 	SpringBone.Initialize(RequiredBones);
 }
 
diff --git a/Engine/Source/Runtime/AnimGraphRuntime/Private/BoneControllers/AnimNode_Trail.cpp b/Engine/Source/Runtime/AnimGraphRuntime/Private/BoneControllers/AnimNode_Trail.cpp
--- a/Engine/Source/Runtime/AnimGraphRuntime/Private/BoneControllers/AnimNode_Trail.cpp
+++ b/Engine/Source/Runtime/AnimGraphRuntime/Private/BoneControllers/AnimNode_Trail.cpp
@@ -44,6 +44,7 @@ FAnimNode_Trail::FAnimNode_Trail()
 
 void FAnimNode_Trail::UpdateInternal(const FAnimationUpdateContext& Context)
 {
+	DECLARE_SCOPE_HIERARCHICAL_COUNTER_ANIMNODE(UpdateInternal)
 	FAnimNode_SkeletalControlBase::UpdateInternal(Context);
 
 	ThisTimstep += Context.GetDeltaTime();
@@ -51,6 +52,7 @@ void FAnimNode_Trail::UpdateInternal(const FAnimationUpdateContext& Context)
 
 void FAnimNode_Trail::GatherDebugData(FNodeDebugData& DebugData)
 {
+	DECLARE_SCOPE_HIERARCHICAL_COUNTER_ANIMNODE(GatherDebugData)
 	FString DebugLine = DebugData.GetNodeName(this);
 
 	DebugLine += "(";
@@ -64,6 +66,7 @@ void FAnimNode_Trail::GatherDebugData(FNodeDebugData& DebugData)
 
 void FAnimNode_Trail::EvaluateSkeletalControl_AnyThread(FComponentSpacePoseContext& Output, TArray<FBoneTransform>& OutBoneTransforms)
 {
+	DECLARE_SCOPE_HIERARCHICAL_COUNTER_ANIMNODE(EvaluateSkeletalControl_AnyThread)
 	SCOPE_CYCLE_COUNTER(STAT_Trail_Eval);
 
 	check(OutBoneTransforms.Num() == 0);
@@ -393,6 +396,7 @@ bool FAnimNode_Trail::IsValidToEvaluate(const USkeleton* Skeleton, const FBoneCo
 
 void FAnimNode_Trail::InitializeBoneReferences(const FBoneContainer& RequiredBones) 
 {
+	DECLARE_SCOPE_HIERARCHICAL_COUNTER_ANIMNODE(InitializeBoneReferences)
 	TrailBone.Initialize(RequiredBones);
 	BaseJoint.Initialize(RequiredBones);
 
@@ -474,6 +478,7 @@ void FAnimNode_Trail::PostLoad()
 
 void FAnimNode_Trail::Initialize_AnyThread(const FAnimationInitializeContext& Context)
 {
+	DECLARE_SCOPE_HIERARCHICAL_COUNTER_ANIMNODE(Initialize_AnyThread)
 	FAnimNode_SkeletalControlBase::Initialize_AnyThread(Context);
 
 	// allocated all memory here in initialize
diff --git a/Engine/Source/Runtime/AnimGraphRuntime/Private/BoneControllers/AnimNode_TwistCorrectiveNode.cpp b/Engine/Source/Runtime/AnimGraphRuntime/Private/BoneControllers/AnimNode_TwistCorrectiveNode.cpp
--- a/Engine/Source/Runtime/AnimGraphRuntime/Private/BoneControllers/AnimNode_TwistCorrectiveNode.cpp
+++ b/Engine/Source/Runtime/AnimGraphRuntime/Private/BoneControllers/AnimNode_TwistCorrectiveNode.cpp
@@ -17,6 +17,7 @@ FAnimNode_TwistCorrectiveNode::FAnimNode_TwistCorrectiveNode()
 
 void FAnimNode_TwistCorrectiveNode::GatherDebugData(FNodeDebugData& DebugData)
 {
+	DECLARE_SCOPE_HIERARCHICAL_COUNTER_ANIMNODE(GatherDebugData)
 	ComponentPose.GatherDebugData(DebugData);
 }
 
@@ -30,6 +31,7 @@ FVector FAnimNode_TwistCorrectiveNode::GetReferenceAxis(FCSPose<FCompactPose>& M
 
 void FAnimNode_TwistCorrectiveNode::EvaluateComponentSpaceInternal(FComponentSpacePoseContext& Context)
 {
+	DECLARE_SCOPE_HIERARCHICAL_COUNTER_ANIMNODE(EvaluateComponentSpaceInternal)
 	// Get the Local space transform and the ref pose transform to see how the transform for the source bone has changed
 	const FBoneContainer& BoneContainer = Context.Pose.GetPose().GetBoneContainer();
 
@@ -57,6 +59,7 @@ bool FAnimNode_TwistCorrectiveNode::IsValidToEvaluate(const USkeleton* Skeleton,
 
 void FAnimNode_TwistCorrectiveNode::InitializeBoneReferences(const FBoneContainer& RequiredBones)
 {
+	DECLARE_SCOPE_HIERARCHICAL_COUNTER_ANIMNODE(InitializeBoneReferences)
 	BaseFrame.Bone.Initialize(RequiredBones);
 	TwistFrame.Bone.Initialize(RequiredBones);
 
@@ -68,12 +71,14 @@ void FAnimNode_TwistCorrectiveNode::InitializeBoneReferences(const FBoneContaine
 
 void FAnimNode_TwistCorrectiveNode::Initialize_AnyThread(const FAnimationInitializeContext& Context)
 {
+	DECLARE_SCOPE_HIERARCHICAL_COUNTER_ANIMNODE(Initialize_AnyThread)
 	FAnimNode_SkeletalControlBase::Initialize_AnyThread(Context);
 	Curve.Initialize(Context.AnimInstanceProxy->GetSkeleton());
 }
 
 void FAnimNode_TwistCorrectiveNode::CacheBones_AnyThread(const FAnimationCacheBonesContext& Context)
 {
+	DECLARE_SCOPE_HIERARCHICAL_COUNTER_ANIMNODE(CacheBones_AnyThread)
 	FAnimNode_SkeletalControlBase::CacheBones_AnyThread(Context);
 	TArray<FTransform> SpaceBases;
 
diff --git a/Engine/Source/Runtime/AnimGraphRuntime/Private/BoneControllers/AnimNode_TwoBoneIK.cpp b/Engine/Source/Runtime/AnimGraphRuntime/Private/BoneControllers/AnimNode_TwoBoneIK.cpp
--- a/Engine/Source/Runtime/AnimGraphRuntime/Private/BoneControllers/AnimNode_TwoBoneIK.cpp
+++ b/Engine/Source/Runtime/AnimGraphRuntime/Private/BoneControllers/AnimNode_TwoBoneIK.cpp
@@ -38,6 +38,7 @@ FAnimNode_TwoBoneIK::FAnimNode_TwoBoneIK()
 
 void FAnimNode_TwoBoneIK::GatherDebugData(FNodeDebugData& DebugData)
 {
+	DECLARE_SCOPE_HIERARCHICAL_COUNTER_ANIMNODE(GatherDebugData)
 	FString DebugLine = DebugData.GetNodeName(this);
 
 	DebugLine += "(";
@@ -68,6 +69,7 @@ FTransform FAnimNode_TwoBoneIK::GetTargetTransform(const FTransform& InComponent
 
 void FAnimNode_TwoBoneIK::EvaluateSkeletalControl_AnyThread(FComponentSpacePoseContext& Output, TArray<FBoneTransform>& OutBoneTransforms)
 {
+	DECLARE_SCOPE_HIERARCHICAL_COUNTER_ANIMNODE(EvaluateSkeletalControl_AnyThread)
 	SCOPE_CYCLE_COUNTER(STAT_TwoBoneIK_Eval);
 
 	check(OutBoneTransforms.Num() == 0);
@@ -221,6 +223,7 @@ bool FAnimNode_TwoBoneIK::IsValidToEvaluate(const USkeleton* Skeleton, const FBo
 
 void FAnimNode_TwoBoneIK::InitializeBoneReferences(const FBoneContainer& RequiredBones) 
 {
+	DECLARE_SCOPE_HIERARCHICAL_COUNTER_ANIMNODE(InitializeBoneReferences)
 	IKBone.Initialize(RequiredBones);
 
 	EffectorTarget.InitializeBoneReferences(RequiredBones);
@@ -242,6 +245,7 @@ void FAnimNode_TwoBoneIK::InitializeBoneReferences(const FBoneContainer& Require
 
 void FAnimNode_TwoBoneIK::Initialize_AnyThread(const FAnimationInitializeContext& Context)
 {
+	DECLARE_SCOPE_HIERARCHICAL_COUNTER_ANIMNODE(Initialize_AnyThread)
 	Super::Initialize_AnyThread(Context);
 	EffectorTarget.Initialize(Context.AnimInstanceProxy);
 	JointTarget.Initialize(Context.AnimInstanceProxy);
diff --git a/Engine/Source/Runtime/Engine/Classes/Animation/AnimNodeBase.h b/Engine/Source/Runtime/Engine/Classes/Animation/AnimNodeBase.h
--- a/Engine/Source/Runtime/Engine/Classes/Animation/AnimNodeBase.h
+++ b/Engine/Source/Runtime/Engine/Classes/Animation/AnimNodeBase.h
@@ -10,8 +10,12 @@
 #include "Animation/AnimCurveTypes.h"
 #include "BonePose.h"
 #include "Logging/TokenizedMessage.h"
+#include "Stats/StatsHierarchical.h"
 #include "AnimNodeBase.generated.h"
 
+#define DECLARE_SCOPE_HIERARCHICAL_COUNTER_ANIMNODE(Method) \
+	DECLARE_SCOPE_HIERARCHICAL_COUNTER_FUNC()
+
 class IAnimClassInterface;
 class UAnimBlueprint;
 class UAnimInstance;
diff --git a/Engine/Source/Runtime/Engine/Classes/EdGraph/EdGraphSchema.h b/Engine/Source/Runtime/Engine/Classes/EdGraph/EdGraphSchema.h
--- a/Engine/Source/Runtime/Engine/Classes/EdGraph/EdGraphSchema.h
+++ b/Engine/Source/Runtime/Engine/Classes/EdGraph/EdGraphSchema.h
@@ -748,14 +748,15 @@ class ENGINE_API UEdGraphSchema : public UObject
 	 *
 	 * @param	Pin			   	The pin on which to set the default value.
 	 * @param	NewDefaultValue	The new default value.
+	 * @param   bMarkAsModified Marks the container of the value as modified
 	 */
-	virtual void TrySetDefaultValue(UEdGraphPin& Pin, const FString& NewDefaultValue) const;
+	virtual void TrySetDefaultValue(UEdGraphPin& Pin, const FString& NewDefaultValue, bool bMarkAsModified = true) const;
 
 	/** Sets the object to the specified pin */
-	virtual void TrySetDefaultObject(UEdGraphPin& Pin, UObject* NewDefaultObject) const;
+	virtual void TrySetDefaultObject(UEdGraphPin& Pin, UObject* NewDefaultObject, bool bMarkAsModified = true) const;
 
 	/** Sets the text to the specified pin */
-	virtual void TrySetDefaultText(UEdGraphPin& InPin, const FText& InNewDefaultText) const;
+	virtual void TrySetDefaultText(UEdGraphPin& InPin, const FText& InNewDefaultText, bool bMarkAsModified = true) const;
 
 	/** Returns if the pin's value matches the given value */
 	virtual bool DoesDefaultValueMatch(const UEdGraphPin& InPin, const FString& InValue) const;
diff --git a/Engine/Source/Runtime/Engine/Private/EdGraph/EdGraphSchema.cpp b/Engine/Source/Runtime/Engine/Private/EdGraph/EdGraphSchema.cpp
--- a/Engine/Source/Runtime/Engine/Private/EdGraph/EdGraphSchema.cpp
+++ b/Engine/Source/Runtime/Engine/Private/EdGraph/EdGraphSchema.cpp
@@ -464,7 +464,7 @@ bool UEdGraphSchema::CreateAutomaticConversionNodeAndConnections(UEdGraphPin* Pi
 	return false;
 }
 
-void UEdGraphSchema::TrySetDefaultValue(UEdGraphPin& Pin, const FString& NewDefaultValue) const
+void UEdGraphSchema::TrySetDefaultValue(UEdGraphPin& Pin, const FString& NewDefaultValue, bool bMarkAsModified) const
 {
 	Pin.DefaultValue = NewDefaultValue;
 
@@ -475,7 +475,7 @@ void UEdGraphSchema::TrySetDefaultValue(UEdGraphPin& Pin, const FString& NewDefa
 #endif	//#if WITH_EDITOR
 }
 
-void UEdGraphSchema::TrySetDefaultObject(UEdGraphPin& Pin, UObject* NewDefaultObject) const
+void UEdGraphSchema::TrySetDefaultObject(UEdGraphPin& Pin, UObject* NewDefaultObject, bool bMarkAsModified) const
 {
 	Pin.DefaultObject = NewDefaultObject;
 
@@ -486,7 +486,7 @@ void UEdGraphSchema::TrySetDefaultObject(UEdGraphPin& Pin, UObject* NewDefaultOb
 #endif	//#if WITH_EDITOR
 }
 
-void UEdGraphSchema::TrySetDefaultText(UEdGraphPin& InPin, const FText& InNewDefaultText) const
+void UEdGraphSchema::TrySetDefaultText(UEdGraphPin& InPin, const FText& InNewDefaultText, bool bMarkAsModified) const
 {
 	InPin.DefaultTextValue = InNewDefaultText;
 
