- Fixed issue of default behavior change due to new feature of blending support on the last joint
 - the var was flipped, so I'm deprecating old one, and adding new one, and flip the meaning

diff --git a/Engine/Source/Editor/AnimGraph/Classes/AnimGraphNode_Trail.h b/Engine/Source/Editor/AnimGraph/Classes/AnimGraphNode_Trail.h
--- a/Engine/Source/Editor/AnimGraph/Classes/AnimGraphNode_Trail.h
+++ b/Engine/Source/Editor/AnimGraph/Classes/AnimGraphNode_Trail.h
@@ -20,6 +20,7 @@ class UAnimGraphNode_Trail : public UAnimGraphNode_SkeletalControlBase
 public:
 	// UObject interface
 	virtual void PostLoad() override;
+	virtual void Serialize(FArchive& Ar) override;
 	// End of UObject interface
 
 	// UEdGraphNode interface
diff --git a/Engine/Source/Editor/AnimGraph/Private/AnimGraphNode_Trail.cpp b/Engine/Source/Editor/AnimGraph/Private/AnimGraphNode_Trail.cpp
--- a/Engine/Source/Editor/AnimGraph/Private/AnimGraphNode_Trail.cpp
+++ b/Engine/Source/Editor/AnimGraph/Private/AnimGraphNode_Trail.cpp
@@ -3,6 +3,7 @@
 #include "AnimGraphNode_Trail.h"
 #include "Components/SkeletalMeshComponent.h"
 #include "Materials/MaterialInstanceDynamic.h"
+#include "UObject/ReleaseObjectVersion.h"
 /////////////////////////////////////////////////////
 // UAnimGraphNode_TrailBone
 
@@ -55,6 +56,22 @@ void UAnimGraphNode_Trail::PostLoad()
 	Node.PostLoad();
 }
 
+void UAnimGraphNode_Trail::Serialize(FArchive& Ar)
+{
+	Super::Serialize(Ar);
+
+	Ar.UsingCustomVersion(FReleaseObjectVersion::GUID);
+
+#if WITH_EDITORONLY_DATA
+	if (Ar.CustomVer(FReleaseObjectVersion::GUID) < FReleaseObjectVersion::TrailNodeBlendVariableNameChange)
+	{
+		if (Node.TrailBoneRotationBlendAlpha_DEPRECATED != 1.f)
+		{
+			Node.LastBoneRotationAnimAlphaBlend = FMath::Clamp<float>(1.f - Node.TrailBoneRotationBlendAlpha_DEPRECATED, 0.f, 1.f);
+		}
+	}
+#endif // #if WITH_EDITORONLY_DATA
+}
 
 void UAnimGraphNode_Trail::CustomizePinData(UEdGraphPin* Pin, FName SourcePropertyName, int32 ArrayIndex) const
 {
diff --git a/Engine/Source/Runtime/AnimGraphRuntime/Private/BoneControllers/AnimNode_Trail.cpp b/Engine/Source/Runtime/AnimGraphRuntime/Private/BoneControllers/AnimNode_Trail.cpp
--- a/Engine/Source/Runtime/AnimGraphRuntime/Private/BoneControllers/AnimNode_Trail.cpp
+++ b/Engine/Source/Runtime/AnimGraphRuntime/Private/BoneControllers/AnimNode_Trail.cpp
@@ -31,7 +31,10 @@ FAnimNode_Trail::FAnimNode_Trail()
 	, RelaxationSpeedScale(1.f)
 	, StretchLimit(0)
 	, FakeVelocity(FVector::ZeroVector)
-	, TrailBoneRotationBlendAlpha(1.f)
+#if WITH_EDITORONLY_DATA
+	, TrailBoneRotationBlendAlpha_DEPRECATED(1.f)
+#endif // WITH_EDITORONLY_DATA
+	, LastBoneRotationAnimAlphaBlend(0.f)
 {
 	FRichCurve* TrailRelaxRichCurve = TrailRelaxationSpeed.GetRichCurve();
 	TrailRelaxRichCurve->AddKey(0.f, 10.f);
@@ -292,7 +295,7 @@ void FAnimNode_Trail::EvaluateSkeletalControl_AnyThread(FComponentSpacePoseConte
 	}
 
 	// For the last bone in the chain, use the rotation from the bone above it.
-	FQuat LeafRotation = FQuat::FastLerp(OutBoneTransforms[ChainLength - 2].Transform.GetRotation(), OutBoneTransforms[ChainLength - 1].Transform.GetRotation(), TrailBoneRotationBlendAlpha);
+	FQuat LeafRotation = FQuat::FastLerp(OutBoneTransforms[ChainLength - 2].Transform.GetRotation(), OutBoneTransforms[ChainLength - 1].Transform.GetRotation(), LastBoneRotationAnimAlphaBlend);
 	LeafRotation.Normalize();
 	OutBoneTransforms[ChainLength - 1].Transform.SetRotation(LeafRotation);
 
diff --git a/Engine/Source/Runtime/AnimGraphRuntime/Public/BoneControllers/AnimNode_Trail.h b/Engine/Source/Runtime/AnimGraphRuntime/Public/BoneControllers/AnimNode_Trail.h
--- a/Engine/Source/Runtime/AnimGraphRuntime/Public/BoneControllers/AnimNode_Trail.h
+++ b/Engine/Source/Runtime/AnimGraphRuntime/Public/BoneControllers/AnimNode_Trail.h
@@ -154,11 +154,15 @@ struct ANIMGRAPHRUNTIME_API FAnimNode_Trail : public FAnimNode_SkeletalControlBa
 	UPROPERTY(EditAnywhere, Category=Velocity)
 	FBoneReference BaseJoint;
 
-	/* How to set last bone rotation. It copies from previous joint if alpha is 1.f, or 0.f will use animated pose 
-	 * This alpha dictates the blend between parent joint and animated pose
+#if WITH_EDITORONLY_DATA
+	UPROPERTY()
+	float TrailBoneRotationBlendAlpha_DEPRECATED;
+#endif // WITH_EDITORONLY_DATA
+	/* How to set last bone rotation. It copies from previous joint if alpha is 0.f, or 1.f will use animated pose 
+	 * This alpha dictates the blend between parent joint and animated pose, and how much you'd like to use animated pose for
 	 */
 	UPROPERTY(EditAnywhere, Category = Rotation, meta = (EditCondition = bReorientParentToChild))
-	float TrailBoneRotationBlendAlpha;
+	float LastBoneRotationAnimAlphaBlend;
 
 	/** Internal use - we need the timestep to do the relaxation in CalculateNewBoneTransforms. */
 	float	ThisTimstep;
diff --git a/Engine/Source/Runtime/Core/Public/UObject/ReleaseObjectVersion.h b/Engine/Source/Runtime/Core/Public/UObject/ReleaseObjectVersion.h
--- a/Engine/Source/Runtime/Core/Public/UObject/ReleaseObjectVersion.h
+++ b/Engine/Source/Runtime/Core/Public/UObject/ReleaseObjectVersion.h
@@ -71,6 +71,9 @@ struct CORE_API FReleaseObjectVersion
 		// Bugfix for FAnimObjectVersion::LinkTimeAnimBlueprintRootDiscovery.
 		LinkTimeAnimBlueprintRootDiscoveryBugFix,
 
+		// Change trail anim node variable deprecation
+		TrailNodeBlendVariableNameChange,
+
 		// -----<new versions can be added above this line>-------------------------------------------------
 		VersionPlusOne,
 		LatestVersion = VersionPlusOne - 1
