Adding nullcheck on get world in UnbindFromSkeletalMeshComponent

diff --git a/Engine/Source/Runtime/AnimGraphRuntime/Private/AnimCustomInstance.cpp b/Engine/Source/Runtime/AnimGraphRuntime/Private/AnimCustomInstance.cpp
--- a/Engine/Source/Runtime/AnimGraphRuntime/Private/AnimCustomInstance.cpp
+++ b/Engine/Source/Runtime/AnimGraphRuntime/Private/AnimCustomInstance.cpp
@@ -46,7 +46,7 @@ void UAnimCustomInstance::UnbindFromSkeletalMeshComponent(USkeletalMeshComponent
 	}
 
 	// if not game world, don't clean this up
-	if (InSkeletalMeshComponent->GetWorld()->IsGameWorld() == false)
+	if (InSkeletalMeshComponent->GetWorld() != nullptr && InSkeletalMeshComponent->GetWorld()->IsGameWorld() == false)
 	{
 		InSkeletalMeshComponent->ClearMotionVector();
 	}
