Making copy pose thread-safe

Moved the work to acquire the bone transforms to PreUpdate on the game thread.
Supported master-pose by checking whether what we are copying is running master-pose (this potentially fixes issues where character heads/hats can detach).
Profiled, this averages ~8us on console per call, so not too bad.

diff --git a/Engine/Source/Runtime/AnimGraphRuntime/Private/AnimNodes/AnimNode_CopyPoseFromMesh.cpp b/Engine/Source/Runtime/AnimGraphRuntime/Private/AnimNodes/AnimNode_CopyPoseFromMesh.cpp
--- a/Engine/Source/Runtime/AnimGraphRuntime/Private/AnimNodes/AnimNode_CopyPoseFromMesh.cpp
+++ b/Engine/Source/Runtime/AnimGraphRuntime/Private/AnimNodes/AnimNode_CopyPoseFromMesh.cpp
@@ -16,7 +16,12 @@ FAnimNode_CopyPoseFromMesh::FAnimNode_CopyPoseFromMesh()
 void FAnimNode_CopyPoseFromMesh::Initialize_AnyThread(const FAnimationInitializeContext& Context)
 {
 	FAnimNode_Base::Initialize_AnyThread(Context);
-	RefreshMeshComponent(Context.AnimInstanceProxy);
+
+	// Initial update of the node, so we dont have a frame-delay on setup
+	FAnimationUpdateContext UpdateContext(Context.AnimInstanceProxy, 0.0f);
+	GetEvaluateGraphExposedInputs().Execute(UpdateContext);
+
+	RefreshMeshComponent(Context.AnimInstanceProxy->GetSkelMeshComponent());
 }
 
 void FAnimNode_CopyPoseFromMesh::CacheBones_AnyThread(const FAnimationCacheBonesContext& Context)
@@ -24,9 +29,9 @@ void FAnimNode_CopyPoseFromMesh::CacheBones_AnyThread(const FAnimationCacheBones
 
 }
 
-void FAnimNode_CopyPoseFromMesh::RefreshMeshComponent(FAnimInstanceProxy* AnimInstanceProxy)
+void FAnimNode_CopyPoseFromMesh::RefreshMeshComponent(USkeletalMeshComponent* TargetMeshComponent)
 {
-	auto ResetMeshComponent = [this](USkeletalMeshComponent* InMeshComponent, FAnimInstanceProxy* InAnimInstanceProxy)
+	auto ResetMeshComponent = [this](USkeletalMeshComponent* InMeshComponent, USkeletalMeshComponent* InTargetMeshComponent)
 	{
 		USkeletalMeshComponent* CurrentMeshComponent = CurrentlyUsedSourceMeshComponent.Get();
 		// if current mesh exists, but not same as input mesh
@@ -35,34 +40,33 @@ void FAnimNode_CopyPoseFromMesh::RefreshMeshComponent(FAnimInstanceProxy* AnimIn
 			// if component has been changed, reinitialize
 			if (CurrentMeshComponent != InMeshComponent)
 			{
-				ReinitializeMeshComponent(InMeshComponent, InAnimInstanceProxy);
+				ReinitializeMeshComponent(InMeshComponent, InTargetMeshComponent);
 			}
 			// if component is still same but mesh has been changed, we have to reinitialize
 			else if (CurrentMeshComponent->SkeletalMesh != CurrentlyUsedSourceMesh.Get())
 			{
-				ReinitializeMeshComponent(InMeshComponent, InAnimInstanceProxy);
+				ReinitializeMeshComponent(InMeshComponent, InTargetMeshComponent);
 			}
 		}
 		// if not valid, but input mesh is
 		else if (!CurrentMeshComponent && InMeshComponent)
 		{
-			ReinitializeMeshComponent(InMeshComponent, InAnimInstanceProxy);
+			ReinitializeMeshComponent(InMeshComponent, InTargetMeshComponent);
 		}
 	};
 
-	if (SourceMeshComponent.IsValid())
+	if(SourceMeshComponent.IsValid())
 	{
-		ResetMeshComponent(SourceMeshComponent.Get(), AnimInstanceProxy);
+		ResetMeshComponent(SourceMeshComponent.Get(), TargetMeshComponent);
 	}
 	else if (bUseAttachedParent)
 	{
-		USkeletalMeshComponent* Component = AnimInstanceProxy->GetSkelMeshComponent();
-		if (Component)
+		if (TargetMeshComponent)
 		{
-			USkeletalMeshComponent* ParentComponent = Cast<USkeletalMeshComponent>(Component->GetAttachParent());
+			USkeletalMeshComponent* ParentComponent = Cast<USkeletalMeshComponent>(TargetMeshComponent->GetAttachParent());
 			if (ParentComponent)
 			{
-				ResetMeshComponent(ParentComponent, AnimInstanceProxy);
+				ResetMeshComponent(ParentComponent, TargetMeshComponent);
 			}
 			else
 			{
@@ -80,11 +84,60 @@ void FAnimNode_CopyPoseFromMesh::RefreshMeshComponent(FAnimInstanceProxy* AnimIn
 	}
 }
 
+void FAnimNode_CopyPoseFromMesh::PreUpdate(const UAnimInstance* InAnimInstance)
+{
+	QUICK_SCOPE_CYCLE_COUNTER(FAnimNode_CopyPoseFromMesh_PreUpdate);
+
+	RefreshMeshComponent(InAnimInstance->GetSkelMeshComponent());
+
+	USkeletalMeshComponent* CurrentMeshComponent = CurrentlyUsedSourceMeshComponent.IsValid() ? CurrentlyUsedSourceMeshComponent.Get() : nullptr;
+
+	if (CurrentMeshComponent && CurrentMeshComponent->SkeletalMesh && CurrentMeshComponent->IsRegistered())
+	{
+		// If our source is running under master-pose, then get bone data from there
+		if(USkeletalMeshComponent* MasterPoseComponent = Cast<USkeletalMeshComponent>(CurrentMeshComponent->MasterPoseComponent.Get()))
+		{
+			CurrentMeshComponent = MasterPoseComponent;
+		}
+
+		// re-check mesh component validity as it may have changed to master
+		if(CurrentMeshComponent->SkeletalMesh && CurrentMeshComponent->IsRegistered())
+		{
+			const bool bUROInSync = CurrentMeshComponent->ShouldUseUpdateRateOptimizations() && CurrentMeshComponent->AnimUpdateRateParams != nullptr && CurrentMeshComponent->AnimUpdateRateParams == InAnimInstance->GetSkelMeshComponent()->AnimUpdateRateParams;
+			const bool bUsingExternalInterpolation = CurrentMeshComponent->IsUsingExternalInterpolation();
+			const TArray<FTransform>& CachedComponentSpaceTransforms = CurrentMeshComponent->GetCachedComponentSpaceTransforms();
+			const bool bArraySizesMatch = CachedComponentSpaceTransforms.Num() == CurrentMeshComponent->GetComponentSpaceTransforms().Num();
+
+			// Copy source array from the appropriate location
+			SourceMeshTransformArray.Reset();
+			SourceMeshTransformArray.Append((bUROInSync || bUsingExternalInterpolation) && bArraySizesMatch ? CachedComponentSpaceTransforms : CurrentMeshComponent->GetComponentSpaceTransforms());
+
+			// Ref skeleton is need for parent index lookups later, so store it now
+			RefSkeleton = &CurrentMeshComponent->SkeletalMesh->RefSkeleton;
+
+			if(bCopyCurves)
+			{
+				UAnimInstance* SourceAnimInstance = CurrentMeshComponent->GetAnimInstance();
+				if (SourceAnimInstance)
+				{
+					// attribute curve contains all list
+					SourceCurveList.Reset();
+					SourceCurveList.Append(SourceAnimInstance->GetAnimationCurveList(EAnimCurveType::AttributeCurve));
+				}
+				else
+				{
+					SourceCurveList.Reset();
+				}
+			}
+		}
+	}
+}
+
 void FAnimNode_CopyPoseFromMesh::Update_AnyThread(const FAnimationUpdateContext& Context)
 {
+	// This introduces a frame of latency in setting the pin-driven source component,
+	// but we cannot do the work to extract transforms on a worker thread as it is not thread safe.
 	GetEvaluateGraphExposedInputs().Execute(Context);
-	// still doing this because if this works, we avoid doing this in eval
-	RefreshMeshComponent(Context.AnimInstanceProxy);
 }
 
 void FAnimNode_CopyPoseFromMesh::Evaluate_AnyThread(FPoseContext& Output)
@@ -92,64 +145,44 @@ void FAnimNode_CopyPoseFromMesh::Evaluate_AnyThread(FPoseContext& Output)
 	FCompactPose& OutPose = Output.Pose;
 	OutPose.ResetToRefPose();
 
-	// I have to refresh this here because some of them already started on worker thread for evaluate
-	// it still has to do update
-	RefreshMeshComponent(Output.AnimInstanceProxy);
-
-	USkeletalMeshComponent* CurrentMeshComponent = CurrentlyUsedSourceMeshComponent.IsValid()? CurrentlyUsedSourceMeshComponent.Get() : nullptr;
-
-	if (CurrentMeshComponent && CurrentMeshComponent->SkeletalMesh && CurrentMeshComponent->IsRegistered())
+	if(SourceMeshTransformArray.Num() > 0 && RefSkeleton != nullptr)
 	{
 		const FBoneContainer& RequiredBones = OutPose.GetBoneContainer();
-		const bool bUROInSync = CurrentMeshComponent->ShouldUseUpdateRateOptimizations() && CurrentMeshComponent->AnimUpdateRateParams != nullptr && CurrentMeshComponent->AnimUpdateRateParams == Output.AnimInstanceProxy->GetSkelMeshComponent()->AnimUpdateRateParams;
-		const bool bUsingExternalInterpolation = CurrentMeshComponent->IsUsingExternalInterpolation();
-		const TArray<FTransform>& CachedComponentSpaceTransforms = CurrentMeshComponent->GetCachedComponentSpaceTransforms();
-		const bool bArraySizesMatch = CachedComponentSpaceTransforms.Num() == CurrentMeshComponent->GetComponentSpaceTransforms().Num();
-		const TArray<FTransform>& SourceMeshTransformArray = (bUROInSync || bUsingExternalInterpolation) && bArraySizesMatch ? CachedComponentSpaceTransforms : CurrentMeshComponent->GetComponentSpaceTransforms();
-		if(SourceMeshTransformArray.Num() > 0)
+
+		for(FCompactPoseBoneIndex PoseBoneIndex : OutPose.ForEachBoneIndex())
 		{
-			for(FCompactPoseBoneIndex PoseBoneIndex : OutPose.ForEachBoneIndex())
+			const int32 SkeletonBoneIndex = RequiredBones.GetSkeletonIndex(PoseBoneIndex);
+			const int32 MeshBoneIndex = RequiredBones.GetSkeletonToPoseBoneIndexArray()[SkeletonBoneIndex];
+			const int32* Value = BoneMapToSource.Find(MeshBoneIndex);
+			if(Value && SourceMeshTransformArray.IsValidIndex(*Value))
 			{
-				const int32& SkeletonBoneIndex = RequiredBones.GetSkeletonIndex(PoseBoneIndex);
-				const int32& MeshBoneIndex = RequiredBones.GetSkeletonToPoseBoneIndexArray()[SkeletonBoneIndex];
-				const int32* Value = BoneMapToSource.Find(MeshBoneIndex);
-				if(Value && SourceMeshTransformArray.IsValidIndex(*Value))
+				const int32 SourceBoneIndex = *Value;
+				const int32 ParentIndex = RefSkeleton->GetParentIndex(SourceBoneIndex);
+				const FCompactPoseBoneIndex MyParentIndex = RequiredBones.GetParentBoneIndex(PoseBoneIndex);
+				// only apply if I also have parent, otherwise, it should apply the space bases
+				if (SourceMeshTransformArray.IsValidIndex(ParentIndex) && MyParentIndex != INDEX_NONE)
 				{
-					const int32 SourceBoneIndex = *Value;
-					const int32 ParentIndex = CurrentMeshComponent->SkeletalMesh->RefSkeleton.GetParentIndex(SourceBoneIndex);
-					const FCompactPoseBoneIndex MyParentIndex = RequiredBones.GetParentBoneIndex(PoseBoneIndex);
-					// only apply if I also have parent, otherwise, it should apply the space bases
-					if (SourceMeshTransformArray.IsValidIndex(ParentIndex) && MyParentIndex != INDEX_NONE)
-					{
-						const FTransform& ParentTransform = SourceMeshTransformArray[ParentIndex];
-						const FTransform& ChildTransform = SourceMeshTransformArray[SourceBoneIndex];
-						OutPose[PoseBoneIndex] = ChildTransform.GetRelativeTransform(ParentTransform);
-					}
-					else
-					{
-						OutPose[PoseBoneIndex] = SourceMeshTransformArray[SourceBoneIndex];
-					}
+					const FTransform& ParentTransform = SourceMeshTransformArray[ParentIndex];
+					const FTransform& ChildTransform = SourceMeshTransformArray[SourceBoneIndex];
+					OutPose[PoseBoneIndex] = ChildTransform.GetRelativeTransform(ParentTransform);
+				}
+				else
+				{
+					OutPose[PoseBoneIndex] = SourceMeshTransformArray[SourceBoneIndex];
 				}
 			}
 		}
+	}
 
-		if (bCopyCurves)
+	if (bCopyCurves)
+	{
+		for (auto Iter = SourceCurveList.CreateConstIterator(); Iter; ++Iter)
 		{
-			UAnimInstance* SourceAnimInstance = CurrentMeshComponent->GetAnimInstance();
-			if (SourceAnimInstance)
+			const SmartName::UID_Type* UID = CurveNameToUIDMap.Find(Iter.Key());
+			if (UID)
 			{
-				// attribute curve contains all list
-				const TMap<FName, float>& SourceCurveList = SourceAnimInstance->GetAnimationCurveList(EAnimCurveType::AttributeCurve);
-
-				for (auto Iter = SourceCurveList.CreateConstIterator(); Iter; ++Iter)
-				{
-					const SmartName::UID_Type* UID = CurveNameToUIDMap.Find(Iter.Key());
-					if (UID)
-					{
-						// set source value to output curve
-						Output.Curve.Set(*UID, Iter.Value());
-					}
-				}
+				// set source value to output curve
+				Output.Curve.Set(*UID, Iter.Value());
 			}
 		}
 	}
@@ -163,7 +196,7 @@ void FAnimNode_CopyPoseFromMesh::GatherDebugData(FNodeDebugData& DebugData)
 	DebugData.AddDebugItem(DebugLine, true);
 }
 
-void FAnimNode_CopyPoseFromMesh::ReinitializeMeshComponent(USkeletalMeshComponent* NewSourceMeshComponent, FAnimInstanceProxy* AnimInstanceProxy)
+void FAnimNode_CopyPoseFromMesh::ReinitializeMeshComponent(USkeletalMeshComponent* NewSourceMeshComponent, USkeletalMeshComponent* TargetMeshComponent)
 {
 	CurrentlyUsedSourceMeshComponent = NewSourceMeshComponent;
 	// reset source mesh
@@ -171,53 +204,49 @@ void FAnimNode_CopyPoseFromMesh::ReinitializeMeshComponent(USkeletalMeshComponen
 	BoneMapToSource.Reset();
 	CurveNameToUIDMap.Reset();
 
-	if (NewSourceMeshComponent && NewSourceMeshComponent->SkeletalMesh && !NewSourceMeshComponent->IsPendingKill())
+	if (TargetMeshComponent && NewSourceMeshComponent && NewSourceMeshComponent->SkeletalMesh && !NewSourceMeshComponent->IsPendingKill())
 	{
-		USkeletalMeshComponent* TargetMeshComponent = AnimInstanceProxy->GetSkelMeshComponent();
-		if (TargetMeshComponent)
-		{
-			USkeletalMesh* SourceSkelMesh = NewSourceMeshComponent->SkeletalMesh;
-			USkeletalMesh* TargetSkelMesh = TargetMeshComponent->SkeletalMesh;
-			CurrentlyUsedSourceMesh = SourceSkelMesh;
+		USkeletalMesh* SourceSkelMesh = NewSourceMeshComponent->SkeletalMesh;
+		USkeletalMesh* TargetSkelMesh = TargetMeshComponent->SkeletalMesh;
+		CurrentlyUsedSourceMesh = SourceSkelMesh;
 
-			if (SourceSkelMesh == TargetSkelMesh)
+		if (SourceSkelMesh == TargetSkelMesh)
+		{
+			for(int32 ComponentSpaceBoneId = 0; ComponentSpaceBoneId < SourceSkelMesh->RefSkeleton.GetNum(); ++ComponentSpaceBoneId)
 			{
-				for(int32 ComponentSpaceBoneId = 0; ComponentSpaceBoneId < SourceSkelMesh->RefSkeleton.GetNum(); ++ComponentSpaceBoneId)
-				{
-					BoneMapToSource.Add(ComponentSpaceBoneId, ComponentSpaceBoneId);
-				}
+				BoneMapToSource.Add(ComponentSpaceBoneId, ComponentSpaceBoneId);
 			}
-			else
+		}
+		else
+		{
+			for (int32 ComponentSpaceBoneId=0; ComponentSpaceBoneId<TargetSkelMesh->RefSkeleton.GetNum(); ++ComponentSpaceBoneId)
 			{
-				for (int32 ComponentSpaceBoneId=0; ComponentSpaceBoneId<TargetSkelMesh->RefSkeleton.GetNum(); ++ComponentSpaceBoneId)
-				{
-					FName BoneName = TargetSkelMesh->RefSkeleton.GetBoneName(ComponentSpaceBoneId);
-					BoneMapToSource.Add(ComponentSpaceBoneId, SourceSkelMesh->RefSkeleton.FindBoneIndex(BoneName));
-				}
+				FName BoneName = TargetSkelMesh->RefSkeleton.GetBoneName(ComponentSpaceBoneId);
+				BoneMapToSource.Add(ComponentSpaceBoneId, SourceSkelMesh->RefSkeleton.FindBoneIndex(BoneName));
 			}
+		}
 
-			if (bCopyCurves)
-			{
-				USkeleton* SourceSkeleton = SourceSkelMesh->Skeleton;
-				USkeleton* TargetSkeleton = TargetSkelMesh->Skeleton;
+		if (bCopyCurves)
+		{
+			USkeleton* SourceSkeleton = SourceSkelMesh->Skeleton;
+			USkeleton* TargetSkeleton = TargetSkelMesh->Skeleton;
 
-				// you shouldn't be here if this happened
-				check(SourceSkeleton && TargetSkeleton);
+			// you shouldn't be here if this happened
+			check(SourceSkeleton && TargetSkeleton);
 
-				const FSmartNameMapping* SourceContainer = SourceSkeleton->GetSmartNameContainer(USkeleton::AnimCurveMappingName);
-				const FSmartNameMapping* TargetContainer = TargetSkeleton->GetSmartNameContainer(USkeleton::AnimCurveMappingName);
+			const FSmartNameMapping* SourceContainer = SourceSkeleton->GetSmartNameContainer(USkeleton::AnimCurveMappingName);
+			const FSmartNameMapping* TargetContainer = TargetSkeleton->GetSmartNameContainer(USkeleton::AnimCurveMappingName);
 
-				TArray<FName> SourceCurveNames;
-				SourceContainer->FillNameArray(SourceCurveNames);
-				for (int32 Index = 0; Index < SourceCurveNames.Num(); ++Index)
+			TArray<FName> SourceCurveNames;
+			SourceContainer->FillNameArray(SourceCurveNames);
+			for (int32 Index = 0; Index < SourceCurveNames.Num(); ++Index)
+			{
+				SmartName::UID_Type UID = TargetContainer->FindUID(SourceCurveNames[Index]);
+				if (UID != SmartName::MaxUID)
 				{
-					SmartName::UID_Type UID = TargetContainer->FindUID(SourceCurveNames[Index]);
-					if (UID != SmartName::MaxUID)
-					{
-						// has a valid UID, add to the list
-						SmartName::UID_Type& Value = CurveNameToUIDMap.Add(SourceCurveNames[Index]);
-						Value = UID;
-					}
+					// has a valid UID, add to the list
+					SmartName::UID_Type& Value = CurveNameToUIDMap.Add(SourceCurveNames[Index]);
+					Value = UID;
 				}
 			}
 		}
diff --git a/Engine/Source/Runtime/AnimGraphRuntime/Public/AnimNodes/AnimNode_CopyPoseFromMesh.h b/Engine/Source/Runtime/AnimGraphRuntime/Public/AnimNodes/AnimNode_CopyPoseFromMesh.h
--- a/Engine/Source/Runtime/AnimGraphRuntime/Public/AnimNodes/AnimNode_CopyPoseFromMesh.h
+++ b/Engine/Source/Runtime/AnimGraphRuntime/Public/AnimNodes/AnimNode_CopyPoseFromMesh.h
@@ -39,6 +39,8 @@ struct ANIMGRAPHRUNTIME_API FAnimNode_CopyPoseFromMesh : public FAnimNode_Base
 	virtual void Update_AnyThread(const FAnimationUpdateContext& Context) override;
 	virtual void Evaluate_AnyThread(FPoseContext& Output) override;
 	virtual void GatherDebugData(FNodeDebugData& DebugData) override;
+	virtual bool HasPreUpdate() const override { return true; }
+	virtual void PreUpdate(const UAnimInstance* InAnimInstance) override;
 	// End of FAnimNode_Base interface
 
 private:
@@ -48,7 +50,16 @@ private:
 	TMap<int32, int32> BoneMapToSource;
 	TMap<FName, SmartName::UID_Type> CurveNameToUIDMap;
 
+	// Cached transforms, copied on the game thread
+	TArray<FTransform> SourceMeshTransformArray;
+
+	// Cached curves, copied on the game thread
+	TMap<FName, float> SourceCurveList;
+
+	// Reference skeleton used for parent bone lookups
+	FReferenceSkeleton* RefSkeleton;
+
 	// reinitialize mesh component 
-	void ReinitializeMeshComponent(USkeletalMeshComponent* NewSkeletalMeshComponent, FAnimInstanceProxy* AnimInstanceProxy);
-	void RefreshMeshComponent(FAnimInstanceProxy* AnimInstanceProxy);
+	void ReinitializeMeshComponent(USkeletalMeshComponent* NewSkeletalMeshComponent, USkeletalMeshComponent* TargetMeshComponent);
+	void RefreshMeshComponent(USkeletalMeshComponent* TargetMeshComponent);
 };
