Back out changelist 4688737

Just in case - found a potential regression with the editor after this went in.

diff --git a/Engine/Source/Runtime/AnimGraphRuntime/Private/AnimNodes/AnimNode_CopyPoseFromMesh.cpp b/Engine/Source/Runtime/AnimGraphRuntime/Private/AnimNodes/AnimNode_CopyPoseFromMesh.cpp
--- a/Engine/Source/Runtime/AnimGraphRuntime/Private/AnimNodes/AnimNode_CopyPoseFromMesh.cpp
+++ b/Engine/Source/Runtime/AnimGraphRuntime/Private/AnimNodes/AnimNode_CopyPoseFromMesh.cpp
@@ -16,12 +16,7 @@ FAnimNode_CopyPoseFromMesh::FAnimNode_CopyPoseFromMesh()
 void FAnimNode_CopyPoseFromMesh::Initialize_AnyThread(const FAnimationInitializeContext& Context)
 {
 	FAnimNode_Base::Initialize_AnyThread(Context);
-
-	// Initial update of the node, so we dont have a frame-delay on setup
-	FAnimationUpdateContext UpdateContext(Context.AnimInstanceProxy, 0.0f);
-	GetEvaluateGraphExposedInputs().Execute(UpdateContext);
-
-	RefreshMeshComponent(Context.AnimInstanceProxy->GetSkelMeshComponent());
+	RefreshMeshComponent(Context.AnimInstanceProxy);
 }
 
 void FAnimNode_CopyPoseFromMesh::CacheBones_AnyThread(const FAnimationCacheBonesContext& Context)
@@ -29,9 +24,9 @@ void FAnimNode_CopyPoseFromMesh::CacheBones_AnyThread(const FAnimationCacheBones
 
 }
 
-void FAnimNode_CopyPoseFromMesh::RefreshMeshComponent(USkeletalMeshComponent* TargetMeshComponent)
+void FAnimNode_CopyPoseFromMesh::RefreshMeshComponent(FAnimInstanceProxy* AnimInstanceProxy)
 {
-	auto ResetMeshComponent = [this](USkeletalMeshComponent* InMeshComponent, USkeletalMeshComponent* InTargetMeshComponent)
+	auto ResetMeshComponent = [this](USkeletalMeshComponent* InMeshComponent, FAnimInstanceProxy* InAnimInstanceProxy)
 	{
 		USkeletalMeshComponent* CurrentMeshComponent = CurrentlyUsedSourceMeshComponent.Get();
 		// if current mesh exists, but not same as input mesh
@@ -40,33 +35,34 @@ void FAnimNode_CopyPoseFromMesh::RefreshMeshComponent(USkeletalMeshComponent* Ta
 			// if component has been changed, reinitialize
 			if (CurrentMeshComponent != InMeshComponent)
 			{
-				ReinitializeMeshComponent(InMeshComponent, InTargetMeshComponent);
+				ReinitializeMeshComponent(InMeshComponent, InAnimInstanceProxy);
 			}
 			// if component is still same but mesh has been changed, we have to reinitialize
 			else if (CurrentMeshComponent->SkeletalMesh != CurrentlyUsedSourceMesh.Get())
 			{
-				ReinitializeMeshComponent(InMeshComponent, InTargetMeshComponent);
+				ReinitializeMeshComponent(InMeshComponent, InAnimInstanceProxy);
 			}
 		}
 		// if not valid, but input mesh is
 		else if (!CurrentMeshComponent && InMeshComponent)
 		{
-			ReinitializeMeshComponent(InMeshComponent, InTargetMeshComponent);
+			ReinitializeMeshComponent(InMeshComponent, InAnimInstanceProxy);
 		}
 	};
 
-	if(SourceMeshComponent.IsValid())
+	if (SourceMeshComponent.IsValid())
 	{
-		ResetMeshComponent(SourceMeshComponent.Get(), TargetMeshComponent);
+		ResetMeshComponent(SourceMeshComponent.Get(), AnimInstanceProxy);
 	}
 	else if (bUseAttachedParent)
 	{
-		if (TargetMeshComponent)
+		USkeletalMeshComponent* Component = AnimInstanceProxy->GetSkelMeshComponent();
+		if (Component)
 		{
-			USkeletalMeshComponent* ParentComponent = Cast<USkeletalMeshComponent>(TargetMeshComponent->GetAttachParent());
+			USkeletalMeshComponent* ParentComponent = Cast<USkeletalMeshComponent>(Component->GetAttachParent());
 			if (ParentComponent)
 			{
-				ResetMeshComponent(ParentComponent, TargetMeshComponent);
+				ResetMeshComponent(ParentComponent, AnimInstanceProxy);
 			}
 			else
 			{
@@ -84,60 +80,11 @@ void FAnimNode_CopyPoseFromMesh::RefreshMeshComponent(USkeletalMeshComponent* Ta
 	}
 }
 
-void FAnimNode_CopyPoseFromMesh::PreUpdate(const UAnimInstance* InAnimInstance)
-{
-	QUICK_SCOPE_CYCLE_COUNTER(FAnimNode_CopyPoseFromMesh_PreUpdate);
-
-	RefreshMeshComponent(InAnimInstance->GetSkelMeshComponent());
-
-	USkeletalMeshComponent* CurrentMeshComponent = CurrentlyUsedSourceMeshComponent.IsValid() ? CurrentlyUsedSourceMeshComponent.Get() : nullptr;
-
-	if (CurrentMeshComponent && CurrentMeshComponent->SkeletalMesh && CurrentMeshComponent->IsRegistered())
-	{
-		// If our source is running under master-pose, then get bone data from there
-		if(USkeletalMeshComponent* MasterPoseComponent = Cast<USkeletalMeshComponent>(CurrentMeshComponent->MasterPoseComponent.Get()))
-		{
-			CurrentMeshComponent = MasterPoseComponent;
-		}
-
-		// re-check mesh component validity as it may have changed to master
-		if(CurrentMeshComponent->SkeletalMesh && CurrentMeshComponent->IsRegistered())
-		{
-			const bool bUROInSync = CurrentMeshComponent->ShouldUseUpdateRateOptimizations() && CurrentMeshComponent->AnimUpdateRateParams != nullptr && CurrentMeshComponent->AnimUpdateRateParams == InAnimInstance->GetSkelMeshComponent()->AnimUpdateRateParams;
-			const bool bUsingExternalInterpolation = CurrentMeshComponent->IsUsingExternalInterpolation();
-			const TArray<FTransform>& CachedComponentSpaceTransforms = CurrentMeshComponent->GetCachedComponentSpaceTransforms();
-			const bool bArraySizesMatch = CachedComponentSpaceTransforms.Num() == CurrentMeshComponent->GetComponentSpaceTransforms().Num();
-
-			// Copy source array from the appropriate location
-			SourceMeshTransformArray.Reset();
-			SourceMeshTransformArray.Append((bUROInSync || bUsingExternalInterpolation) && bArraySizesMatch ? CachedComponentSpaceTransforms : CurrentMeshComponent->GetComponentSpaceTransforms());
-
-			// Ref skeleton is need for parent index lookups later, so store it now
-			RefSkeleton = &CurrentMeshComponent->SkeletalMesh->RefSkeleton;
-
-			if(bCopyCurves)
-			{
-				UAnimInstance* SourceAnimInstance = CurrentMeshComponent->GetAnimInstance();
-				if (SourceAnimInstance)
-				{
-					// attribute curve contains all list
-					SourceCurveList.Reset();
-					SourceCurveList.Append(SourceAnimInstance->GetAnimationCurveList(EAnimCurveType::AttributeCurve));
-				}
-				else
-				{
-					SourceCurveList.Reset();
-				}
-			}
-		}
-	}
-}
-
 void FAnimNode_CopyPoseFromMesh::Update_AnyThread(const FAnimationUpdateContext& Context)
 {
-	// This introduces a frame of latency in setting the pin-driven source component,
-	// but we cannot do the work to extract transforms on a worker thread as it is not thread safe.
 	GetEvaluateGraphExposedInputs().Execute(Context);
+	// still doing this because if this works, we avoid doing this in eval
+	RefreshMeshComponent(Context.AnimInstanceProxy);
 }
 
 void FAnimNode_CopyPoseFromMesh::Evaluate_AnyThread(FPoseContext& Output)
@@ -145,44 +92,64 @@ void FAnimNode_CopyPoseFromMesh::Evaluate_AnyThread(FPoseContext& Output)
 	FCompactPose& OutPose = Output.Pose;
 	OutPose.ResetToRefPose();
 
-	if(SourceMeshTransformArray.Num() > 0 && RefSkeleton != nullptr)
+	// I have to refresh this here because some of them already started on worker thread for evaluate
+	// it still has to do update
+	RefreshMeshComponent(Output.AnimInstanceProxy);
+
+	USkeletalMeshComponent* CurrentMeshComponent = CurrentlyUsedSourceMeshComponent.IsValid()? CurrentlyUsedSourceMeshComponent.Get() : nullptr;
+
+	if (CurrentMeshComponent && CurrentMeshComponent->SkeletalMesh && CurrentMeshComponent->IsRegistered())
 	{
 		const FBoneContainer& RequiredBones = OutPose.GetBoneContainer();
-
-		for(FCompactPoseBoneIndex PoseBoneIndex : OutPose.ForEachBoneIndex())
+		const bool bUROInSync = CurrentMeshComponent->ShouldUseUpdateRateOptimizations() && CurrentMeshComponent->AnimUpdateRateParams != nullptr && CurrentMeshComponent->AnimUpdateRateParams == Output.AnimInstanceProxy->GetSkelMeshComponent()->AnimUpdateRateParams;
+		const bool bUsingExternalInterpolation = CurrentMeshComponent->IsUsingExternalInterpolation();
+		const TArray<FTransform>& CachedComponentSpaceTransforms = CurrentMeshComponent->GetCachedComponentSpaceTransforms();
+		const bool bArraySizesMatch = CachedComponentSpaceTransforms.Num() == CurrentMeshComponent->GetComponentSpaceTransforms().Num();
+		const TArray<FTransform>& SourceMeshTransformArray = (bUROInSync || bUsingExternalInterpolation) && bArraySizesMatch ? CachedComponentSpaceTransforms : CurrentMeshComponent->GetComponentSpaceTransforms();
+		if(SourceMeshTransformArray.Num() > 0)
 		{
-			const int32 SkeletonBoneIndex = RequiredBones.GetSkeletonIndex(PoseBoneIndex);
-			const int32 MeshBoneIndex = RequiredBones.GetSkeletonToPoseBoneIndexArray()[SkeletonBoneIndex];
-			const int32* Value = BoneMapToSource.Find(MeshBoneIndex);
-			if(Value && SourceMeshTransformArray.IsValidIndex(*Value))
+			for(FCompactPoseBoneIndex PoseBoneIndex : OutPose.ForEachBoneIndex())
 			{
-				const int32 SourceBoneIndex = *Value;
-				const int32 ParentIndex = RefSkeleton->GetParentIndex(SourceBoneIndex);
-				const FCompactPoseBoneIndex MyParentIndex = RequiredBones.GetParentBoneIndex(PoseBoneIndex);
-				// only apply if I also have parent, otherwise, it should apply the space bases
-				if (SourceMeshTransformArray.IsValidIndex(ParentIndex) && MyParentIndex != INDEX_NONE)
+				const int32& SkeletonBoneIndex = RequiredBones.GetSkeletonIndex(PoseBoneIndex);
+				const int32& MeshBoneIndex = RequiredBones.GetSkeletonToPoseBoneIndexArray()[SkeletonBoneIndex];
+				const int32* Value = BoneMapToSource.Find(MeshBoneIndex);
+				if(Value && SourceMeshTransformArray.IsValidIndex(*Value))
 				{
-					const FTransform& ParentTransform = SourceMeshTransformArray[ParentIndex];
-					const FTransform& ChildTransform = SourceMeshTransformArray[SourceBoneIndex];
-					OutPose[PoseBoneIndex] = ChildTransform.GetRelativeTransform(ParentTransform);
-				}
-				else
-				{
-					OutPose[PoseBoneIndex] = SourceMeshTransformArray[SourceBoneIndex];
+					const int32 SourceBoneIndex = *Value;
+					const int32 ParentIndex = CurrentMeshComponent->SkeletalMesh->RefSkeleton.GetParentIndex(SourceBoneIndex);
+					const FCompactPoseBoneIndex MyParentIndex = RequiredBones.GetParentBoneIndex(PoseBoneIndex);
+					// only apply if I also have parent, otherwise, it should apply the space bases
+					if (SourceMeshTransformArray.IsValidIndex(ParentIndex) && MyParentIndex != INDEX_NONE)
+					{
+						const FTransform& ParentTransform = SourceMeshTransformArray[ParentIndex];
+						const FTransform& ChildTransform = SourceMeshTransformArray[SourceBoneIndex];
+						OutPose[PoseBoneIndex] = ChildTransform.GetRelativeTransform(ParentTransform);
+					}
+					else
+					{
+						OutPose[PoseBoneIndex] = SourceMeshTransformArray[SourceBoneIndex];
+					}
 				}
 			}
 		}
-	}
 
-	if (bCopyCurves)
-	{
-		for (auto Iter = SourceCurveList.CreateConstIterator(); Iter; ++Iter)
+		if (bCopyCurves)
 		{
-			const SmartName::UID_Type* UID = CurveNameToUIDMap.Find(Iter.Key());
-			if (UID)
+			UAnimInstance* SourceAnimInstance = CurrentMeshComponent->GetAnimInstance();
+			if (SourceAnimInstance)
 			{
-				// set source value to output curve
-				Output.Curve.Set(*UID, Iter.Value());
+				// attribute curve contains all list
+				const TMap<FName, float>& SourceCurveList = SourceAnimInstance->GetAnimationCurveList(EAnimCurveType::AttributeCurve);
+
+				for (auto Iter = SourceCurveList.CreateConstIterator(); Iter; ++Iter)
+				{
+					const SmartName::UID_Type* UID = CurveNameToUIDMap.Find(Iter.Key());
+					if (UID)
+					{
+						// set source value to output curve
+						Output.Curve.Set(*UID, Iter.Value());
+					}
+				}
 			}
 		}
 	}
@@ -196,7 +163,7 @@ void FAnimNode_CopyPoseFromMesh::GatherDebugData(FNodeDebugData& DebugData)
 	DebugData.AddDebugItem(DebugLine, true);
 }
 
-void FAnimNode_CopyPoseFromMesh::ReinitializeMeshComponent(USkeletalMeshComponent* NewSourceMeshComponent, USkeletalMeshComponent* TargetMeshComponent)
+void FAnimNode_CopyPoseFromMesh::ReinitializeMeshComponent(USkeletalMeshComponent* NewSourceMeshComponent, FAnimInstanceProxy* AnimInstanceProxy)
 {
 	CurrentlyUsedSourceMeshComponent = NewSourceMeshComponent;
 	// reset source mesh
@@ -204,49 +171,53 @@ void FAnimNode_CopyPoseFromMesh::ReinitializeMeshComponent(USkeletalMeshComponen
 	BoneMapToSource.Reset();
 	CurveNameToUIDMap.Reset();
 
-	if (TargetMeshComponent && NewSourceMeshComponent && NewSourceMeshComponent->SkeletalMesh && !NewSourceMeshComponent->IsPendingKill())
+	if (NewSourceMeshComponent && NewSourceMeshComponent->SkeletalMesh && !NewSourceMeshComponent->IsPendingKill())
 	{
-		USkeletalMesh* SourceSkelMesh = NewSourceMeshComponent->SkeletalMesh;
-		USkeletalMesh* TargetSkelMesh = TargetMeshComponent->SkeletalMesh;
-		CurrentlyUsedSourceMesh = SourceSkelMesh;
-
-		if (SourceSkelMesh == TargetSkelMesh)
+		USkeletalMeshComponent* TargetMeshComponent = AnimInstanceProxy->GetSkelMeshComponent();
+		if (TargetMeshComponent)
 		{
-			for(int32 ComponentSpaceBoneId = 0; ComponentSpaceBoneId < SourceSkelMesh->RefSkeleton.GetNum(); ++ComponentSpaceBoneId)
+			USkeletalMesh* SourceSkelMesh = NewSourceMeshComponent->SkeletalMesh;
+			USkeletalMesh* TargetSkelMesh = TargetMeshComponent->SkeletalMesh;
+			CurrentlyUsedSourceMesh = SourceSkelMesh;
+
+			if (SourceSkelMesh == TargetSkelMesh)
 			{
-				BoneMapToSource.Add(ComponentSpaceBoneId, ComponentSpaceBoneId);
+				for(int32 ComponentSpaceBoneId = 0; ComponentSpaceBoneId < SourceSkelMesh->RefSkeleton.GetNum(); ++ComponentSpaceBoneId)
+				{
+					BoneMapToSource.Add(ComponentSpaceBoneId, ComponentSpaceBoneId);
+				}
 			}
-		}
-		else
-		{
-			for (int32 ComponentSpaceBoneId=0; ComponentSpaceBoneId<TargetSkelMesh->RefSkeleton.GetNum(); ++ComponentSpaceBoneId)
+			else
 			{
-				FName BoneName = TargetSkelMesh->RefSkeleton.GetBoneName(ComponentSpaceBoneId);
-				BoneMapToSource.Add(ComponentSpaceBoneId, SourceSkelMesh->RefSkeleton.FindBoneIndex(BoneName));
+				for (int32 ComponentSpaceBoneId=0; ComponentSpaceBoneId<TargetSkelMesh->RefSkeleton.GetNum(); ++ComponentSpaceBoneId)
+				{
+					FName BoneName = TargetSkelMesh->RefSkeleton.GetBoneName(ComponentSpaceBoneId);
+					BoneMapToSource.Add(ComponentSpaceBoneId, SourceSkelMesh->RefSkeleton.FindBoneIndex(BoneName));
+				}
 			}
-		}
 
-		if (bCopyCurves)
-		{
-			USkeleton* SourceSkeleton = SourceSkelMesh->Skeleton;
-			USkeleton* TargetSkeleton = TargetSkelMesh->Skeleton;
+			if (bCopyCurves)
+			{
+				USkeleton* SourceSkeleton = SourceSkelMesh->Skeleton;
+				USkeleton* TargetSkeleton = TargetSkelMesh->Skeleton;
 
-			// you shouldn't be here if this happened
-			check(SourceSkeleton && TargetSkeleton);
+				// you shouldn't be here if this happened
+				check(SourceSkeleton && TargetSkeleton);
 
-			const FSmartNameMapping* SourceContainer = SourceSkeleton->GetSmartNameContainer(USkeleton::AnimCurveMappingName);
-			const FSmartNameMapping* TargetContainer = TargetSkeleton->GetSmartNameContainer(USkeleton::AnimCurveMappingName);
+				const FSmartNameMapping* SourceContainer = SourceSkeleton->GetSmartNameContainer(USkeleton::AnimCurveMappingName);
+				const FSmartNameMapping* TargetContainer = TargetSkeleton->GetSmartNameContainer(USkeleton::AnimCurveMappingName);
 
-			TArray<FName> SourceCurveNames;
-			SourceContainer->FillNameArray(SourceCurveNames);
-			for (int32 Index = 0; Index < SourceCurveNames.Num(); ++Index)
-			{
-				SmartName::UID_Type UID = TargetContainer->FindUID(SourceCurveNames[Index]);
-				if (UID != SmartName::MaxUID)
+				TArray<FName> SourceCurveNames;
+				SourceContainer->FillNameArray(SourceCurveNames);
+				for (int32 Index = 0; Index < SourceCurveNames.Num(); ++Index)
 				{
-					// has a valid UID, add to the list
-					SmartName::UID_Type& Value = CurveNameToUIDMap.Add(SourceCurveNames[Index]);
-					Value = UID;
+					SmartName::UID_Type UID = TargetContainer->FindUID(SourceCurveNames[Index]);
+					if (UID != SmartName::MaxUID)
+					{
+						// has a valid UID, add to the list
+						SmartName::UID_Type& Value = CurveNameToUIDMap.Add(SourceCurveNames[Index]);
+						Value = UID;
+					}
 				}
 			}
 		}
diff --git a/Engine/Source/Runtime/AnimGraphRuntime/Public/AnimNodes/AnimNode_CopyPoseFromMesh.h b/Engine/Source/Runtime/AnimGraphRuntime/Public/AnimNodes/AnimNode_CopyPoseFromMesh.h
--- a/Engine/Source/Runtime/AnimGraphRuntime/Public/AnimNodes/AnimNode_CopyPoseFromMesh.h
+++ b/Engine/Source/Runtime/AnimGraphRuntime/Public/AnimNodes/AnimNode_CopyPoseFromMesh.h
@@ -39,8 +39,6 @@ struct ANIMGRAPHRUNTIME_API FAnimNode_CopyPoseFromMesh : public FAnimNode_Base
 	virtual void Update_AnyThread(const FAnimationUpdateContext& Context) override;
 	virtual void Evaluate_AnyThread(FPoseContext& Output) override;
 	virtual void GatherDebugData(FNodeDebugData& DebugData) override;
-	virtual bool HasPreUpdate() const override { return true; }
-	virtual void PreUpdate(const UAnimInstance* InAnimInstance) override;
 	// End of FAnimNode_Base interface
 
 private:
@@ -50,16 +48,7 @@ private:
 	TMap<int32, int32> BoneMapToSource;
 	TMap<FName, SmartName::UID_Type> CurveNameToUIDMap;
 
-	// Cached transforms, copied on the game thread
-	TArray<FTransform> SourceMeshTransformArray;
-
-	// Cached curves, copied on the game thread
-	TMap<FName, float> SourceCurveList;
-
-	// Reference skeleton used for parent bone lookups
-	FReferenceSkeleton* RefSkeleton;
-
 	// reinitialize mesh component 
-	void ReinitializeMeshComponent(USkeletalMeshComponent* NewSkeletalMeshComponent, USkeletalMeshComponent* TargetMeshComponent);
-	void RefreshMeshComponent(USkeletalMeshComponent* TargetMeshComponent);
+	void ReinitializeMeshComponent(USkeletalMeshComponent* NewSkeletalMeshComponent, FAnimInstanceProxy* AnimInstanceProxy);
+	void RefreshMeshComponent(FAnimInstanceProxy* AnimInstanceProxy);
 };
