Moved CreateConstraint out of disable collision if statement

================================

diff --git a/Engine/Source/Runtime/AnimGraphRuntime/Private/BoneControllers/AnimNode_RigidBody.cpp b/Engine/Source/Runtime/AnimGraphRuntime/Private/BoneControllers/AnimNode_RigidBody.cpp
--- a/Engine/Source/Runtime/AnimGraphRuntime/Private/BoneControllers/AnimNode_RigidBody.cpp
+++ b/Engine/Source/Runtime/AnimGraphRuntime/Private/BoneControllers/AnimNode_RigidBody.cpp
@@ -719,9 +719,9 @@ void FAnimNode_RigidBody::InitPhysics(const UAnimInstance* InAnimInstance)
 #if WITH_CHAOS || WITH_IMMEDIATE_PHYSX || PHYSICS_INTERFACE_LLIMMEDIATE
                         ensure(false);
 #else
-							if (bForceDisableCollisionBetweenConstraintBodies)
+						PhysicsSimulation->CreateJoint(ConstraintRef.ConstraintData, Body1Handle, Body2Handle);
+						if (bForceDisableCollisionBetweenConstraintBodies)
 						{
-							PhysicsSimulation->CreateJoint(ConstraintRef.ConstraintData, Body1Handle, Body2Handle);
 							int32 BodyIndex1 = UsePhysicsAsset->FindBodyIndex(CI->ConstraintBone1);
 							int32 BodyIndex2 = UsePhysicsAsset->FindBodyIndex(CI->ConstraintBone2);
 							if (BodyIndex1 != INDEX_NONE && BodyIndex2 != INDEX_NONE)
