Change default value for FAnimNode_RigidBody::bForceDisableCollisionBetweenConstraintBodies
to be false in order to not break existing assets.

diff --git a/Engine/Source/Runtime/AnimGraphRuntime/Private/BoneControllers/AnimNode_RigidBody.cpp b/Engine/Source/Runtime/AnimGraphRuntime/Private/BoneControllers/AnimNode_RigidBody.cpp
--- a/Engine/Source/Runtime/AnimGraphRuntime/Private/BoneControllers/AnimNode_RigidBody.cpp
+++ b/Engine/Source/Runtime/AnimGraphRuntime/Private/BoneControllers/AnimNode_RigidBody.cpp
@@ -56,7 +56,7 @@ FAnimNode_RigidBody::FAnimNode_RigidBody():
 	ComponentLinearAccScale = FVector::ZeroVector;
 	ComponentLinearVelScale = FVector::ZeroVector;
 	ComponentAppliedLinearAccClamp = FVector(10000,10000,10000);
-	bForceDisableCollisionBetweenConstraintBodies = true;
+	bForceDisableCollisionBetweenConstraintBodies = false;
 }
 
 FAnimNode_RigidBody::~FAnimNode_RigidBody()
