Fixed setting name by pin in PoseByName node

From this UDN: https://udn.unrealengine.com/questions/477742/bug-incorrect-fname-value-from-c-to-blueprint-on-p.html
Before this fix, only the name set at startup would be used.

diff --git a/Engine/Source/Runtime/AnimGraphRuntime/Private/AnimNodes/AnimNode_PoseByName.cpp b/Engine/Source/Runtime/AnimGraphRuntime/Private/AnimNodes/AnimNode_PoseByName.cpp
--- a/Engine/Source/Runtime/AnimGraphRuntime/Private/AnimNodes/AnimNode_PoseByName.cpp
+++ b/Engine/Source/Runtime/AnimGraphRuntime/Private/AnimNodes/AnimNode_PoseByName.cpp
@@ -27,6 +27,13 @@ void FAnimNode_PoseByName::RebuildPoseList(const FBoneContainer& InBoneContainer
 void FAnimNode_PoseByName::UpdateAssetPlayer(const FAnimationUpdateContext& Context)
 {
 	FAnimNode_PoseHandler::UpdateAssetPlayer(Context);
+
+	// update pose extraction context if the name differs
+	if (CurrentPoseName != PoseName)
+	{
+		RebuildPoseList(Context.AnimInstanceProxy->GetRequiredBones(), PoseAsset);
+		CurrentPoseName = PoseName;
+	}
 }
 
 void FAnimNode_PoseByName::Evaluate_AnyThread(FPoseContext& Output)
diff --git a/Engine/Source/Runtime/AnimGraphRuntime/Public/AnimNodes/AnimNode_PoseByName.h b/Engine/Source/Runtime/AnimGraphRuntime/Public/AnimNodes/AnimNode_PoseByName.h
--- a/Engine/Source/Runtime/AnimGraphRuntime/Public/AnimNodes/AnimNode_PoseByName.h
+++ b/Engine/Source/Runtime/AnimGraphRuntime/Public/AnimNodes/AnimNode_PoseByName.h
@@ -21,9 +21,14 @@ public:
 	UPROPERTY(EditAnywhere, BlueprintReadWrite, Category = Settings, meta = (PinHiddenByDefault))
 	float PoseWeight;
 
+private:
+	/** Cached current pose name used for invalidation */
+	FName CurrentPoseName;
+
 public:	
 	FAnimNode_PoseByName()
 		: PoseWeight(1.f)
+		, CurrentPoseName(NAME_None)
 	{
 	}
 
