support correct sourcepose for curvesource node

https://udn.unrealengine.com/questions/407109/view.html

#code review:Thomas.Sarkanen

diff --git a/Engine/Source/Runtime/AnimGraphRuntime/Private/AnimNodes/AnimNode_CurveSource.cpp b/Engine/Source/Runtime/AnimGraphRuntime/Private/AnimNodes/AnimNode_CurveSource.cpp
--- a/Engine/Source/Runtime/AnimGraphRuntime/Private/AnimNodes/AnimNode_CurveSource.cpp
+++ b/Engine/Source/Runtime/AnimGraphRuntime/Private/AnimNodes/AnimNode_CurveSource.cpp
@@ -99,4 +99,23 @@ void FAnimNode_CurveSource::Update_AnyThread(const FAnimationUpdateContext& Cont
 {
 	// Evaluate any BP logic plugged into this node
 	GetEvaluateGraphExposedInputs().Execute(Context);
+	SourcePose.Update(Context);
 }
+
+void FAnimNode_CurveSource::Initialize_AnyThread(const FAnimationInitializeContext& Context)
+{
+	FAnimNode_Base::Initialize_AnyThread(Context);
+	SourcePose.Initialize(Context);
+}
+
+void FAnimNode_CurveSource::CacheBones_AnyThread(const FAnimationCacheBonesContext& Context)
+{
+	FAnimNode_Base::CacheBones_AnyThread(Context);
+	SourcePose.CacheBones(Context);
+}
+
+void FAnimNode_CurveSource::GatherDebugData(FNodeDebugData& DebugData)
+{
+	FAnimNode_Base::GatherDebugData(DebugData);
+	SourcePose.GatherDebugData(DebugData.BranchFlow(1.f));
+}
\ No newline at end of file
diff --git a/Engine/Source/Runtime/AnimGraphRuntime/Public/AnimNodes/AnimNode_CurveSource.h b/Engine/Source/Runtime/AnimGraphRuntime/Public/AnimNodes/AnimNode_CurveSource.h
--- a/Engine/Source/Runtime/AnimGraphRuntime/Public/AnimNodes/AnimNode_CurveSource.h
+++ b/Engine/Source/Runtime/AnimGraphRuntime/Public/AnimNodes/AnimNode_CurveSource.h
@@ -45,5 +45,8 @@ struct ANIMGRAPHRUNTIME_API FAnimNode_CurveSource : public FAnimNode_Base
 	virtual void PreUpdate(const UAnimInstance* InAnimInstance) override;
 	virtual void Evaluate_AnyThread(FPoseContext& Output) override;
 	virtual void Update_AnyThread(const FAnimationUpdateContext& Context) override;
+	virtual void Initialize_AnyThread(const FAnimationInitializeContext& Context) override;
+	virtual void CacheBones_AnyThread(const FAnimationCacheBonesContext& Context) override;
+	virtual void GatherDebugData(FNodeDebugData& DebugData) override;
 	// End of FAnimNode_Base interface
 };
