Removed NaN check that was still getting compiled in Test/Shipping

diff --git a/Engine/Source/Runtime/AnimGraphRuntime/Private/BoneControllers/AnimNode_SkeletalControlBase.cpp b/Engine/Source/Runtime/AnimGraphRuntime/Private/BoneControllers/AnimNode_SkeletalControlBase.cpp
--- a/Engine/Source/Runtime/AnimGraphRuntime/Private/BoneControllers/AnimNode_SkeletalControlBase.cpp
+++ b/Engine/Source/Runtime/AnimGraphRuntime/Private/BoneControllers/AnimNode_SkeletalControlBase.cpp
@@ -141,8 +141,11 @@ void FAnimNode_SkeletalControlBase::EvaluateComponentSpace_AnyThread(FComponentS
 	// save current pose before applying skeletal control to compute the exact gizmo location in AnimGraphNode
 	ForwardedPose.CopyPose(Output.Pose);
 #endif // #if WITH_EDITORONLY_DATA
+
+#if DO_CHECK
 	// this is to ensure Source data does not contain NaN
 	ensure(Output.ContainsNaN() == false);
+#endif
 
 	// Apply the skeletal control if it's valid
 	if (FAnimWeight::IsRelevant(ActualAlpha) && IsValidToEvaluate(Output.AnimInstanceProxy->GetSkeleton(), Output.AnimInstanceProxy->GetRequiredBones()))
