Added three new blueprint nodes:

1) An easing node which can take a float value and run it through most of the easing functions available in FMath
2) A node to record motions over time and output an averaged velocity. This is useful for tracking the velocity of say a weapon and drive secondary effects.
3) A node which is a monolithic large version of 2) to make it easy for riggers to track the velocity of a bone / socket.

diff --git a/Engine/Source/Runtime/AnimGraphRuntime/Private/KismetAnimationLibrary.cpp b/Engine/Source/Runtime/AnimGraphRuntime/Private/KismetAnimationLibrary.cpp
--- a/Engine/Source/Runtime/AnimGraphRuntime/Private/KismetAnimationLibrary.cpp
+++ b/Engine/Source/Runtime/AnimGraphRuntime/Private/KismetAnimationLibrary.cpp
@@ -107,5 +107,134 @@ float UKismetAnimationLibrary::K2_MakePerlinNoiseAndRemap(float Value, float Ran
 	// perlin noise output is always from [-1, 1]
 	return FMath::GetMappedRangeValueClamped(FVector2D(-1.f, 1.f), FVector2D(RangeOutMin, RangeOutMax), FMath::PerlinNoise1D(Value));
 }
+
+float UKismetAnimationLibrary::K2_CalculateVelocityFromPositionHistory(
+	float DeltaSeconds,
+	FVector Position,
+	UPARAM(ref) FPositionHistory& History,
+	int32 NumberOfSamples,
+	float VelocityMin,
+	float VelocityMax
+) {
+	NumberOfSamples = FMath::Max<uint32>(NumberOfSamples, 2);
+	if (DeltaSeconds <= 0.0f)
+	{
+		return 0.f;
+	}
+
+	// if the number of samples changes down clear the history
+	if (History.Positions.Num() > NumberOfSamples)
+	{
+		History.Positions.Reset();
+		History.Velocities.Reset();
+		History.LastIndex = 0;
+	}
+
+	// append to the history until it's full and then loop around when filling it 
+	// to reuse the memory
+	if (History.Positions.Num() == 0)
+	{
+		History.Positions.Reserve(NumberOfSamples);
+		History.Velocities.Reserve(NumberOfSamples);
+		History.Positions.Add(Position);
+		History.Velocities.Add(0.f);
+		History.LastIndex = 0;
+		return 0.f;
+	}
+	else
+	{
+		float LengthOfV = ((Position - History.Positions[History.LastIndex]) / DeltaSeconds).Size();
+
+		if (History.Positions.Num() == NumberOfSamples)
+		{
+			int32 NextIndex = History.LastIndex + 1;
+			if (NextIndex == History.Positions.Num())
+			{
+				NextIndex = 0;
+			}
+			History.Positions[NextIndex] = Position;
+			History.Velocities[History.LastIndex] = LengthOfV;
+			History.LastIndex = NextIndex;
+		}
+		else
+		{
+			History.LastIndex = History.Positions.Num();
+			History.Positions.Add(Position);
+			History.Velocities.Add(LengthOfV);
+		}
+	}
+
+	// compute average velocity
+	float LengthOfV = 0.0f;
+	for (int32 i = 0; i < History.Velocities.Num(); i++)
+	{
+		LengthOfV += History.Velocities[i];
+	}
+
+	// Avoids NaN due to the FMath::Max instruction above.
+	LengthOfV /= float(History.Velocities.Num());
+
+	if (VelocityMin < 0.0f || VelocityMax < 0.0f || VelocityMax <= VelocityMin)
+	{
+		return LengthOfV;
+	}
+
+	// Avoids NaN due to the condition above.
+	return FMath::Clamp((LengthOfV - VelocityMin) / (VelocityMax - VelocityMin), 0.f, 1.f);
+}
+
+float UKismetAnimationLibrary::K2_ScalarEasing(float Value, EEasingFuncType EasingType)
+{
+	switch(EasingType)
+	{
+		case EEasingFuncType::Linear:			return FMath::Clamp<float>(Value, 0.f, 1.f);
+		case EEasingFuncType::Sinusoidal:		return FMath::Clamp<float>((FMath::Sin(Value * PI - HALF_PI) + 1.f) / 2.f, 0.f, 1.f);
+		case EEasingFuncType::Cubic:			return FMath::Clamp<float>(FMath::CubicInterp<float>(0.f, 0.f, 1.f, 0.f, Value), 0.f, 1.f);
+		case EEasingFuncType::QuadraticInOut:	return FMath::Clamp<float>(FMath::InterpEaseInOut<float>(0.f, 1.f, Value, 2), 0.f, 1.f);
+		case EEasingFuncType::CubicInOut:		return FMath::Clamp<float>(FMath::InterpEaseInOut<float>(0.f, 1.f, Value, 3), 0.f, 1.f);
+		case EEasingFuncType::HermiteCubic:		return FMath::Clamp<float>(FMath::SmoothStep(0.0f, 1.0f, Value), 0.0f, 1.0f);
+		case EEasingFuncType::QuarticInOut:		return FMath::Clamp<float>(FMath::InterpEaseInOut<float>(0.f, 1.f, Value, 4), 0.f, 1.f);
+		case EEasingFuncType::QuinticInOut:		return FMath::Clamp<float>(FMath::InterpEaseInOut<float>(0.f, 1.f, Value, 5), 0.f, 1.f);
+		case EEasingFuncType::CircularIn:		return FMath::Clamp<float>(FMath::InterpCircularIn<float>(0.0f, 1.0f, Value), 0.0f, 1.0f);
+		case EEasingFuncType::CircularOut:		return FMath::Clamp<float>(FMath::InterpCircularOut<float>(0.0f, 1.0f, Value), 0.0f, 1.0f);
+		case EEasingFuncType::CircularInOut:	return FMath::Clamp<float>(FMath::InterpCircularInOut<float>(0.0f, 1.0f, Value), 0.0f, 1.0f);
+		case EEasingFuncType::ExpIn:			return FMath::Clamp<float>(FMath::InterpExpoIn<float>(0.0f, 1.0f, Value), 0.0f, 1.0f);
+		case EEasingFuncType::ExpOut:			return FMath::Clamp<float>(FMath::InterpExpoOut<float>(0.0f, 1.0f, Value), 0.0f, 1.0f);
+		case EEasingFuncType::ExpInOut:			return FMath::Clamp<float>(FMath::InterpExpoInOut<float>(0.0f, 1.0f, Value), 0.0f, 1.0f);
+	}
+	return FMath::Clamp<float>(Value, 0.f, 1.f);;
+}
+
+float UKismetAnimationLibrary::K2_CalculateVelocityFromSockets(
+	float DeltaSeconds,
+	USkeletalMeshComponent * Component,
+	const FName SocketOrBoneName,
+	const FName FrameOfReference,
+	ERelativeTransformSpace SocketSpace,
+	FVector OffsetInBoneSpace,
+	UPARAM(ref) FPositionHistory& History,
+	int32 NumberOfSamples,
+	float VelocityMin,
+	float VelocityMax,
+	EEasingFuncType EasingType
+) {
+	if (Component && SocketOrBoneName != NAME_None)
+	{
+		FTransform SocketTransform = Component->GetSocketTransform(SocketOrBoneName, SocketSpace);
+		if (FrameOfReference != NAME_None)
+		{
+			// make the bone's / socket's transform relative to the frame of reference.
+			FTransform FrameOfReferenceTransform = Component->GetSocketTransform(FrameOfReference, SocketSpace);
+			SocketTransform = SocketTransform.GetRelativeTransform(FrameOfReferenceTransform);
+		}
+
+		FVector Position = SocketTransform.TransformPosition(OffsetInBoneSpace);
+		float Velocity = K2_CalculateVelocityFromPositionHistory(DeltaSeconds, Position, History, NumberOfSamples, VelocityMin, VelocityMax);
+		return K2_ScalarEasing(Velocity, EasingType);
+	}
+
+	return VelocityMin;
+}
+
 #undef LOCTEXT_NAMESPACE
 
diff --git a/Engine/Source/Runtime/AnimGraphRuntime/Public/KismetAnimationLibrary.h b/Engine/Source/Runtime/AnimGraphRuntime/Public/KismetAnimationLibrary.h
--- a/Engine/Source/Runtime/AnimGraphRuntime/Public/KismetAnimationLibrary.h
+++ b/Engine/Source/Runtime/AnimGraphRuntime/Public/KismetAnimationLibrary.h
@@ -5,6 +5,7 @@
 #include "CoreMinimal.h"
 #include "Kismet/BlueprintFunctionLibrary.h"
 #include "Components/SceneComponent.h"
+#include "KismetAnimationTypes.h"
 #include "KismetAnimationLibrary.generated.h"
 
 class USkeletalMeshComponent;
@@ -32,5 +33,56 @@ class ANIMGRAPHRUNTIME_API UKismetAnimationLibrary : public UBlueprintFunctionLi
 
 	UFUNCTION(BlueprintPure, Category = "Utilities|Animation", meta = (DisplayName = "Make Perlin Noise and Remap", ScriptName = "MakeFloatFromPerlinNoise", RangeOutMinX = "-1.f", RangeOutMaxX = "1.f", RangeOutMinY = "-1.f", RangeOutMaxY = "1.f", RangeOutMinZ = "-1.f", RangeOutMaxZ = "1.f"))
 	static float K2_MakePerlinNoiseAndRemap(float Value, float RangeOutMin, float RangeOutMax);
+
+	/** 
+		This function calculates the velocity of a position changing over time.
+		You need to hook up a valid PositionHistory variable to this for storage.
+
+		@param Position The position to track over time.
+		@param History The history to use for storage.
+		@param NumberOfSamples The number of samples to use for the history. The higher the number of samples - the smoother the velocity changes.
+		@param VelocityMin The minimum velocity to use for normalization (if both min and max are set to 0, normalization is turned off)
+		@param VelocityMax The maximum velocity to use for normalization (if both min and max are set to 0, normalization is turned off)
+	*/
+	UFUNCTION(BlueprintCallable, Category = "Utilities|Animation", meta = (DisplayName = "Calculate Velocity From Position History", ScriptName = "CalculateVelocityFromPositionHistory", NumberOfSamples = "16", VelocityMin = "0.f", VelocityMax = "128.f"))
+	static float K2_CalculateVelocityFromPositionHistory(
+		float DeltaSeconds,
+		FVector Position,
+		UPARAM(ref) FPositionHistory& History,
+		int32 NumberOfSamples,
+		float VelocityMin,
+		float VelocityMax
+	);
+
+	UFUNCTION(BlueprintPure, Category = "Utilities|Animation", meta = (DisplayName = "Scalar Easing", ScriptName = "ScalarEasing", Value = "0.f"))
+	static float K2_ScalarEasing(float Value, EEasingFuncType EasingType);
+
+	/** 
+		This function calculates the velocity of an offset position on a bone / socket over time.
+		The bone's / socket's motion can be expressed within a reference frame (another bone / socket). 
+		You need to hook up a valid PositionHistory variable to this for storage.
+
+		@param SocketOrBoneName The name of the bone / socket to track.
+		@param ReferenceSocketOrBone The name of the bone / socket to use as a frame of reference (or None if no frame of reference == world space).
+		@param OffsetInBoneSpace The relative position in the space of the bone / socket to track over time.
+		@param History The history to use for storage.
+		@param NumberOfSamples The number of samples to use for the history. The higher the number of samples - the smoother the velocity changes.
+		@param VelocityMin The minimum velocity to use for normalization (if both min and max are set to 0, normalization is turned off)
+		@param VelocityMax The maximum velocity to use for normalization (if both min and max are set to 0, normalization is turned off)
+	*/
+	UFUNCTION(BlueprintCallable, Category = "Utilities|Animation", meta = (DisplayName = "Calculate Velocity From Sockets", ScriptName = "CalculateVelocityFromSockets", NumberOfSamples = "16", VelocityMin = "0.f", VelocityMax = "128.f"))
+	static float K2_CalculateVelocityFromSockets(
+		float DeltaSeconds, 
+		USkeletalMeshComponent * Component,
+		const FName SocketOrBoneName,
+		const FName ReferenceSocketOrBone,
+		ERelativeTransformSpace SocketSpace,
+		FVector OffsetInBoneSpace, 
+		UPARAM(ref) FPositionHistory& History, 
+		int32 NumberOfSamples,
+		float VelocityMin, 
+		float VelocityMax,
+		EEasingFuncType EasingType
+	);
 };
 
diff --git a/Engine/Source/Runtime/AnimGraphRuntime/Public/KismetAnimationTypes.h b/Engine/Source/Runtime/AnimGraphRuntime/Public/KismetAnimationTypes.h
new file mode 100644
--- /dev/null
+++ b/Engine/Source/Runtime/AnimGraphRuntime/Public/KismetAnimationTypes.h
@@ -0,0 +1,49 @@
+// Copyright 1998-2018 Epic Games, Inc. All Rights Reserved.
+
+#pragma once
+
+#include "CoreMinimal.h"
+#include "KismetAnimationTypes.generated.h"
+
+// The kind of easing function to use
+UENUM()
+enum class EEasingFuncType : uint8
+{
+	Linear,
+	Sinusoidal,
+	Cubic,
+	QuadraticInOut,
+	CubicInOut,
+	HermiteCubic,
+	QuarticInOut,
+	QuinticInOut,
+	CircularIn,
+	CircularOut,
+	CircularInOut,
+	ExpIn,
+	ExpOut,
+	ExpInOut
+};
+
+USTRUCT(BlueprintType)
+struct FPositionHistory
+{
+	GENERATED_BODY()
+
+public:
+
+	/** Default constructor */
+	FPositionHistory()
+		: Positions(TArray<FVector>())
+		, Velocities(TArray<float>())
+		, LastIndex(0)
+	{}
+
+	/** The recorded positions */
+	UPROPERTY(EditAnywhere, Category = "FPositionHistory", meta = (MetaClass = "AdvancedCopyCustomization"))
+	TArray<FVector> Positions;
+
+	TArray<float> Velocities;
+	uint32 LastIndex; /// The last index used to store a position
+};
+
