Fix copyright issue in KismetAnimationTypesHeader

diff --git a/Engine/Source/Runtime/AnimGraphRuntime/Public/KismetAnimationTypes.h b/Engine/Source/Runtime/AnimGraphRuntime/Public/KismetAnimationTypes.h
--- a/Engine/Source/Runtime/AnimGraphRuntime/Public/KismetAnimationTypes.h
+++ b/Engine/Source/Runtime/AnimGraphRuntime/Public/KismetAnimationTypes.h
@@ -1,4 +1,4 @@
-// Copyright 1998-2018 Epic Games, Inc. All Rights Reserved.
+// Copyright 1998-2019 Epic Games, Inc. All Rights Reserved.
 
 #pragma once
 
