Validate BaseBoneRef in Update FAnimNode_RigidBody::IsValidToEvaluate(). Should fix crashes caused by referencing bones that don't exist at some LODs.

diff --git a/Engine/Source/Runtime/AnimGraphRuntime/Private/BoneControllers/AnimNode_RigidBody.cpp b/Engine/Source/Runtime/AnimGraphRuntime/Private/BoneControllers/AnimNode_RigidBody.cpp
--- a/Engine/Source/Runtime/AnimGraphRuntime/Private/BoneControllers/AnimNode_RigidBody.cpp
+++ b/Engine/Source/Runtime/AnimGraphRuntime/Private/BoneControllers/AnimNode_RigidBody.cpp
@@ -1077,4 +1077,9 @@ void FAnimNode_RigidBody::PostSerialize(const FArchive& Ar)
 }
 #endif
 
+bool FAnimNode_RigidBody::IsValidToEvaluate(const USkeleton* Skeleton, const FBoneContainer& RequiredBones)
+{
+	return BaseBoneRef.IsValidToEvaluate(RequiredBones);
+}
+
 #undef LOCTEXT_NAMESPACE
diff --git a/Engine/Source/Runtime/AnimGraphRuntime/Public/BoneControllers/AnimNode_RigidBody.h b/Engine/Source/Runtime/AnimGraphRuntime/Public/BoneControllers/AnimNode_RigidBody.h
--- a/Engine/Source/Runtime/AnimGraphRuntime/Public/BoneControllers/AnimNode_RigidBody.h
+++ b/Engine/Source/Runtime/AnimGraphRuntime/Public/BoneControllers/AnimNode_RigidBody.h
@@ -53,7 +53,7 @@ struct ANIMGRAPHRUNTIME_API FAnimNode_RigidBody : public FAnimNode_SkeletalContr
 	virtual void PreUpdate(const UAnimInstance* InAnimInstance) override;
 	virtual void UpdateInternal(const FAnimationUpdateContext& Context) override;
 	virtual bool HasPreUpdate() const override { return true; }
-	virtual bool IsValidToEvaluate(const USkeleton* Skeleton, const FBoneContainer& RequiredBones) override { return true; }
+	virtual bool IsValidToEvaluate(const USkeleton* Skeleton, const FBoneContainer& RequiredBones) override;
 	virtual bool NeedsDynamicReset() const override;
 	virtual void ResetDynamics(ETeleportType InTeleportType) override;
 	virtual int32 GetLODThreshold() const override;
