Merging using Fortnight-Main<->Dev-Anim

diff --git a/Engine/Source/Runtime/AnimGraphRuntime/Private/BoneControllers/AnimNode_AnimDynamics.cpp b/Engine/Source/Runtime/AnimGraphRuntime/Private/BoneControllers/AnimNode_AnimDynamics.cpp
--- a/Engine/Source/Runtime/AnimGraphRuntime/Private/BoneControllers/AnimNode_AnimDynamics.cpp
+++ b/Engine/Source/Runtime/AnimGraphRuntime/Private/BoneControllers/AnimNode_AnimDynamics.cpp
@@ -4,6 +4,7 @@
 #include "GameFramework/WorldSettings.h"
 #include "Animation/AnimInstanceProxy.h"
 #include "PhysicsEngine/PhysicsSettings.h"
+#include "CommonAnimationLibrary.h"
 
 DEFINE_STAT(STAT_AnimDynamicsOverall);
 DEFINE_STAT(STAT_AnimDynamicsWindData);
@@ -152,6 +153,7 @@ FAnimNode_AnimDynamics::FAnimNode_AnimDynamics()
 , bLinearSpring(false)
 , bAngularSpring(false)
 , bChain(false)
+, RetargetingSettings(FRotationRetargetingInfo(false /* enabled */))
 #if ENABLE_ANIM_DRAW_DEBUG
 , FilteredBoneIndex(INDEX_NONE)
 #endif
@@ -350,6 +352,34 @@ void FAnimNode_AnimDynamics::EvaluateSkeletalControl_AnyThread(FComponentSpacePo
 
 				FTransform NewBoneTransform(CurrentBody.Pose.Orientation, CurrentBody.Pose.Position + CurrentBody.Pose.Orientation.RotateVector(JointOffsets[Idx]));
 
+				if (RetargetingSettings.bEnabled)
+				{
+					FTransform ParentTransform = FTransform::Identity;
+					FCompactPoseBoneIndex ParentBoneIndex = BoneContainer.GetParentBoneIndex(BoneIndex);
+					if (ParentBoneIndex != INDEX_NONE)
+					{
+						ParentTransform = GetBoneTransformInSimSpace(Output, ParentBoneIndex);
+					}
+
+					FQuat RetargetedRotation = CommonAnimationLibrary::RetargetSingleRotation(
+						NewBoneTransform.GetRotation(),
+						RetargetingSettings.Source * ParentTransform,
+						RetargetingSettings.Target * ParentTransform,
+						RetargetingSettings.CustomCurve,
+						RetargetingSettings.EasingType,
+						RetargetingSettings.bFlipEasing,
+						RetargetingSettings.EasingWeight,
+						RetargetingSettings.RotationComponent,
+						RetargetingSettings.TwistAxis,
+						RetargetingSettings.bUseAbsoluteAngle,
+						RetargetingSettings.SourceMinimum,
+						RetargetingSettings.SourceMaximum,
+						RetargetingSettings.TargetMinimum,
+						RetargetingSettings.TargetMaximum);
+
+					NewBoneTransform.SetRotation(RetargetedRotation);
+				}
+
 				NewBoneTransform = GetComponentSpaceTransformFromSimSpace(SimulationSpace, Output, NewBoneTransform);
 
 				OutBoneTransforms.Add(FBoneTransform(BoneIndex, NewBoneTransform));
diff --git a/Engine/Source/Runtime/AnimGraphRuntime/Private/CommonAnimationLibrary.cpp b/Engine/Source/Runtime/AnimGraphRuntime/Private/CommonAnimationLibrary.cpp
new file mode 100644
--- /dev/null
+++ b/Engine/Source/Runtime/AnimGraphRuntime/Private/CommonAnimationLibrary.cpp
@@ -0,0 +1,300 @@
+// Copyright 1998-2019 Epic Games, Inc. All Rights Reserved.
+
+#include "CommonAnimationLibrary.h"
+#include "AnimationCoreLibrary.h"
+
+#define LOCTEXT_NAMESPACE "CommonAnimationLibrary"
+
+//////////////////////////////////////////////////////////////////////////
+
+float CommonAnimationLibrary::ScalarEasing(float Value, const FRuntimeFloatCurve& CustomCurve, EEasingFuncType EasingType, bool bFlip, float Weight)
+{
+	float Original = Value;
+	float Result = Value;
+
+	if(bFlip)
+	{
+		Value = 1.f - Value;
+	}
+
+	switch(EasingType)
+	{
+		case EEasingFuncType::Linear:
+		{
+			Result = FMath::Clamp<float>(Value, 0.f, 1.f);
+			break;
+		}
+		case EEasingFuncType::Sinusoidal:
+		{
+			Result = FMath::Clamp<float>((FMath::Sin(Value * PI - HALF_PI) + 1.f) / 2.f, 0.f, 1.f);
+			break;
+		}
+		case EEasingFuncType::Cubic:
+		{
+			Result = FMath::Clamp<float>(FMath::CubicInterp<float>(0.f, 0.f, 1.f, 0.f, Value), 0.f, 1.f);
+			break;
+		}
+		case EEasingFuncType::QuadraticInOut:
+		{
+			Result = FMath::Clamp<float>(FMath::InterpEaseInOut<float>(0.f, 1.f, Value, 2), 0.f, 1.f);
+			break;
+		}
+		case EEasingFuncType::CubicInOut:
+		{
+			Result = FMath::Clamp<float>(FMath::InterpEaseInOut<float>(0.f, 1.f, Value, 3), 0.f, 1.f);
+			break;
+		}
+		case EEasingFuncType::HermiteCubic:
+		{
+			Result = FMath::Clamp<float>(FMath::SmoothStep(0.0f, 1.0f, Value), 0.0f, 1.0f);
+			break;
+		}
+		case EEasingFuncType::QuarticInOut:
+		{
+			Result = FMath::Clamp<float>(FMath::InterpEaseInOut<float>(0.f, 1.f, Value, 4), 0.f, 1.f);
+			break;
+		}
+		case EEasingFuncType::QuinticInOut:
+		{
+			Result = FMath::Clamp<float>(FMath::InterpEaseInOut<float>(0.f, 1.f, Value, 5), 0.f, 1.f);
+			break;
+		}
+		case EEasingFuncType::CircularIn:
+		{
+			Result = FMath::Clamp<float>(FMath::InterpCircularIn<float>(0.0f, 1.0f, Value), 0.0f, 1.0f);
+			break;
+		}
+		case EEasingFuncType::CircularOut:
+		{
+			Result = FMath::Clamp<float>(FMath::InterpCircularOut<float>(0.0f, 1.0f, Value), 0.0f, 1.0f);
+			break;
+		}
+		case EEasingFuncType::CircularInOut:
+		{
+			Result = FMath::Clamp<float>(FMath::InterpCircularInOut<float>(0.0f, 1.0f, Value), 0.0f, 1.0f);
+			break;
+		}
+		case EEasingFuncType::ExpIn:
+		{
+			Result = FMath::Clamp<float>(FMath::InterpExpoIn<float>(0.0f, 1.0f, Value), 0.0f, 1.0f);
+			break;
+		}
+		case EEasingFuncType::ExpOut:
+		{
+			Result = FMath::Clamp<float>(FMath::InterpExpoOut<float>(0.0f, 1.0f, Value), 0.0f, 1.0f);
+			break;
+		}
+		case EEasingFuncType::ExpInOut:
+		{
+			Result = FMath::Clamp<float>(FMath::InterpExpoInOut<float>(0.0f, 1.0f, Value), 0.0f, 1.0f);
+			break;
+		}
+		case EEasingFuncType::CustomCurve:
+		{
+			if (CustomCurve.GetRichCurveConst() != nullptr)
+			{
+				Result = CustomCurve.GetRichCurveConst()->Eval(Value, Value);
+			}
+			break;
+		}
+	}
+
+	if(bFlip)
+	{
+		Result = 1.f - Result;
+	}
+
+	return FMath::Lerp<float>(Original, Result, Weight);
+}
+
+FVector CommonAnimationLibrary::RetargetSingleLocation(
+	FVector Location,
+	FTransform Source,
+	FTransform Target,
+	const FRuntimeFloatCurve& CustomCurve,
+	EEasingFuncType EasingType,
+	bool bFlipEasing,
+	float EasingWeight,
+	FVector Axis,
+	float SourceMinimum,
+	float SourceMaximum,
+	float TargetMinimum,
+	float TargetMaximum
+)
+{
+	float SourceDelta = SourceMaximum - SourceMinimum;
+	float TargetDelta = TargetMaximum - TargetMinimum;
+	if (FMath::IsNearlyEqual(SourceDelta, 0.f) || FMath::IsNearlyEqual(TargetDelta, 0.f))
+	{
+		return Location;
+	}
+
+	float AxisLengthSquared = Axis.SizeSquared();
+	if (FMath::IsNearlyEqual(AxisLengthSquared, 0.f))
+	{
+		Axis = FVector(1.f, 0.f, 0.f);
+	}
+	else if (!FMath::IsNearlyEqual(AxisLengthSquared, 1.f))
+	{
+		// normalize only in this case
+		Axis = Axis * (1.f / sqrt(AxisLengthSquared));
+	}
+
+	Location = Source.InverseTransformPosition(Location);
+
+	float Delta = FVector::DotProduct(Location, Axis);
+	float OutBias = (Delta - SourceMinimum) / SourceDelta;
+	if (EasingType != EEasingFuncType::Linear)
+	{
+		OutBias = ScalarEasing(OutBias, CustomCurve, EasingType, bFlipEasing, EasingWeight);
+	}
+	float TargetBias = FMath::Lerp<float>(TargetMinimum, TargetMaximum, OutBias);
+
+	Location = Location + Axis * (TargetBias - Delta);
+
+	return Target.TransformPosition(Location);
+}
+
+FQuat CommonAnimationLibrary::RetargetSingleRotation(
+	FQuat Rotation,
+	FTransform Source,
+	FTransform Target,
+	const FRuntimeFloatCurve& CustomCurve,
+	EEasingFuncType EasingType,
+	bool bFlipEasing,
+	float EasingWeight,
+	ERotationComponent RotationComponent,
+	FVector TwistAxis,
+	bool bUseAbsoluteAngle,
+	float SourceMinimum,
+	float SourceMaximum,
+	float TargetMinimum,
+	float TargetMaximum
+)
+{
+	float SourceDelta = SourceMaximum - SourceMinimum;
+	float TargetDelta = TargetMaximum - TargetMinimum;
+	if (FMath::IsNearlyEqual(SourceDelta, 0.f) || FMath::IsNearlyEqual(TargetDelta, 0.f))
+	{
+		return Rotation;
+	}
+
+	Rotation = Source.InverseTransformRotation(Rotation);
+
+	float Angle = 0.f;
+	FQuat Swing, Twist;
+	FVector Euler;
+
+	switch (RotationComponent)
+	{
+		case ERotationComponent::EulerX:
+		{
+			Euler = Rotation.Euler();
+			Angle = Euler.X;
+			break;
+		}
+		case ERotationComponent::EulerY:
+		{
+			Euler = Rotation.Euler();
+			Angle = Euler.Y;
+			break;
+		}
+		case ERotationComponent::EulerZ:
+		{
+			Euler = Rotation.Euler();
+			Angle = Euler.Z;
+			break;
+		}
+		case ERotationComponent::QuaternionAngle:
+		{
+			Angle = FMath::RadiansToDegrees(Rotation.GetAngle());
+			break;
+		}
+		case ERotationComponent::SwingAngle:
+		case ERotationComponent::TwistAngle:
+		{
+			float AxisLengthSquared = TwistAxis.SizeSquared();
+			if (FMath::IsNearlyEqual(AxisLengthSquared, 0.f))
+			{
+				TwistAxis = FVector(1.f, 0.f, 0.f);
+			}
+			else
+			{
+				TwistAxis.Normalize();
+			}
+
+			Rotation.ToSwingTwist(TwistAxis, Swing, Twist);
+			if (RotationComponent == ERotationComponent::SwingAngle)
+			{
+				Angle = FMath::RadiansToDegrees(Swing.GetAngle());
+			}
+			else
+			{
+				Angle = FMath::RadiansToDegrees(Twist.GetAngle());
+			}
+			break;
+		}
+	}
+
+	float AngleSign = Angle < 0.f ? -1.f : 1.f;
+	if (bUseAbsoluteAngle)
+	{
+		Angle = FMath::Abs<float>(Angle);
+	}
+
+	float OutBias = (Angle - SourceMinimum) / SourceDelta;
+	float TempBias = OutBias;
+	if (EasingType != EEasingFuncType::Linear)
+	{
+		OutBias = ScalarEasing(OutBias, CustomCurve, EasingType, bFlipEasing, EasingWeight);
+	}
+
+	float TargetAngle = FMath::Lerp<float>(TargetMinimum, TargetMaximum, OutBias);
+	if (bUseAbsoluteAngle)
+	{
+		TargetAngle *= AngleSign;
+	}
+
+	switch (RotationComponent)
+	{
+		case ERotationComponent::EulerX:
+		{
+			Euler.X = TargetAngle;
+			Rotation = FQuat::MakeFromEuler(Euler);
+			break;
+		}
+		case ERotationComponent::EulerY:
+		{
+			Euler.Y = TargetAngle;
+			Rotation = FQuat::MakeFromEuler(Euler);
+			break;
+		}
+		case ERotationComponent::EulerZ:
+		{
+			Euler.Z = TargetAngle;
+			Rotation = FQuat::MakeFromEuler(Euler);
+			break;
+		}
+		case ERotationComponent::QuaternionAngle:
+		{
+			Rotation = FQuat(Rotation.GetRotationAxis(), FMath::DegreesToRadians(TargetAngle));
+			break;
+		}
+		case ERotationComponent::SwingAngle:
+		{
+			Rotation = Twist * FQuat(Swing.GetRotationAxis(), FMath::DegreesToRadians(TargetAngle));
+			break;
+		}
+		case ERotationComponent::TwistAngle:
+		{
+			Rotation = FQuat(Twist.GetRotationAxis(), FMath::DegreesToRadians(TargetAngle)) * Swing;
+			break;
+		}
+	}
+
+	Rotation.Normalize();
+	return Target.TransformRotation(Rotation);
+}
+
+#undef LOCTEXT_NAMESPACE
+
diff --git a/Engine/Source/Runtime/AnimGraphRuntime/Private/KismetAnimationLibrary.cpp b/Engine/Source/Runtime/AnimGraphRuntime/Private/KismetAnimationLibrary.cpp
--- a/Engine/Source/Runtime/AnimGraphRuntime/Private/KismetAnimationLibrary.cpp
+++ b/Engine/Source/Runtime/AnimGraphRuntime/Private/KismetAnimationLibrary.cpp
@@ -1,7 +1,7 @@
 // Copyright 1998-2019 Epic Games, Inc. All Rights Reserved.
 
 #include "KismetAnimationLibrary.h"
-#include "AnimationCoreLibrary.h"
+#include "CommonAnimationLibrary.h"
 #include "AnimationCoreLibrary.h"
 #include "Blueprint/BlueprintSupport.h"
 #include "Components/SkeletalMeshComponent.h"
@@ -183,28 +183,6 @@ float UKismetAnimationLibrary::K2_CalculateVelocityFromPositionHistory(
 	return FMath::Clamp((LengthOfV - VelocityMin) / (VelocityMax - VelocityMin), 0.f, 1.f);
 }
 
-float UKismetAnimationLibrary::K2_ScalarEasing(float Value, EEasingFuncType EasingType)
-{
-	switch(EasingType)
-	{
-		case EEasingFuncType::Linear:			return FMath::Clamp<float>(Value, 0.f, 1.f);
-		case EEasingFuncType::Sinusoidal:		return FMath::Clamp<float>((FMath::Sin(Value * PI - HALF_PI) + 1.f) / 2.f, 0.f, 1.f);
-		case EEasingFuncType::Cubic:			return FMath::Clamp<float>(FMath::CubicInterp<float>(0.f, 0.f, 1.f, 0.f, Value), 0.f, 1.f);
-		case EEasingFuncType::QuadraticInOut:	return FMath::Clamp<float>(FMath::InterpEaseInOut<float>(0.f, 1.f, Value, 2), 0.f, 1.f);
-		case EEasingFuncType::CubicInOut:		return FMath::Clamp<float>(FMath::InterpEaseInOut<float>(0.f, 1.f, Value, 3), 0.f, 1.f);
-		case EEasingFuncType::HermiteCubic:		return FMath::Clamp<float>(FMath::SmoothStep(0.0f, 1.0f, Value), 0.0f, 1.0f);
-		case EEasingFuncType::QuarticInOut:		return FMath::Clamp<float>(FMath::InterpEaseInOut<float>(0.f, 1.f, Value, 4), 0.f, 1.f);
-		case EEasingFuncType::QuinticInOut:		return FMath::Clamp<float>(FMath::InterpEaseInOut<float>(0.f, 1.f, Value, 5), 0.f, 1.f);
-		case EEasingFuncType::CircularIn:		return FMath::Clamp<float>(FMath::InterpCircularIn<float>(0.0f, 1.0f, Value), 0.0f, 1.0f);
-		case EEasingFuncType::CircularOut:		return FMath::Clamp<float>(FMath::InterpCircularOut<float>(0.0f, 1.0f, Value), 0.0f, 1.0f);
-		case EEasingFuncType::CircularInOut:	return FMath::Clamp<float>(FMath::InterpCircularInOut<float>(0.0f, 1.0f, Value), 0.0f, 1.0f);
-		case EEasingFuncType::ExpIn:			return FMath::Clamp<float>(FMath::InterpExpoIn<float>(0.0f, 1.0f, Value), 0.0f, 1.0f);
-		case EEasingFuncType::ExpOut:			return FMath::Clamp<float>(FMath::InterpExpoOut<float>(0.0f, 1.0f, Value), 0.0f, 1.0f);
-		case EEasingFuncType::ExpInOut:			return FMath::Clamp<float>(FMath::InterpExpoInOut<float>(0.0f, 1.0f, Value), 0.0f, 1.0f);
-	}
-	return FMath::Clamp<float>(Value, 0.f, 1.f);;
-}
-
 float UKismetAnimationLibrary::K2_CalculateVelocityFromSockets(
 	float DeltaSeconds,
 	USkeletalMeshComponent * Component,
@@ -216,7 +194,8 @@ float UKismetAnimationLibrary::K2_CalculateVelocityFromSockets(
 	int32 NumberOfSamples,
 	float VelocityMin,
 	float VelocityMax,
-	EEasingFuncType EasingType
+	EEasingFuncType EasingType,
+	const FRuntimeFloatCurve& CustomCurve
 ) {
 	if (Component && SocketOrBoneName != NAME_None)
 	{
@@ -230,7 +209,7 @@ float UKismetAnimationLibrary::K2_CalculateVelocityFromSockets(
 
 		FVector Position = SocketTransform.TransformPosition(OffsetInBoneSpace);
 		float Velocity = K2_CalculateVelocityFromPositionHistory(DeltaSeconds, Position, History, NumberOfSamples, VelocityMin, VelocityMax);
-		return K2_ScalarEasing(Velocity, EasingType);
+		return CommonAnimationLibrary::ScalarEasing(Velocity, CustomCurve, EasingType);
 	}
 
 	return VelocityMin;
diff --git a/Engine/Source/Runtime/AnimGraphRuntime/Public/BoneControllers/AnimNode_AnimDynamics.h b/Engine/Source/Runtime/AnimGraphRuntime/Public/BoneControllers/AnimNode_AnimDynamics.h
--- a/Engine/Source/Runtime/AnimGraphRuntime/Public/BoneControllers/AnimNode_AnimDynamics.h
+++ b/Engine/Source/Runtime/AnimGraphRuntime/Public/BoneControllers/AnimNode_AnimDynamics.h
@@ -10,6 +10,7 @@
 #include "BonePose.h"
 #include "BoneControllers/AnimNode_SkeletalControlBase.h"
 #include "Animation/AnimInstanceProxy.h"
+#include "KismetAnimationTypes.h"
 #include "AnimNode_AnimDynamics.generated.h"
 
 class UAnimInstance;
@@ -400,6 +401,10 @@ struct ANIMGRAPHRUNTIME_API FAnimNode_AnimDynamics : public FAnimNode_SkeletalCo
 	UPROPERTY(EditAnywhere, Category = Setup, meta=(DisplayAfter="RelativeSpaceBone"))
 	uint8 bChain:1;
 
+	/** The settings for rotation retargeting */
+	UPROPERTY(EditAnywhere, Category = Retargeting)
+	FRotationRetargetingInfo RetargetingSettings;
+
 	// FAnimNode_SkeletalControlBase interface
 	virtual void Initialize_AnyThread(const FAnimationInitializeContext& Context) override;
 	virtual void UpdateInternal(const FAnimationUpdateContext& Context) override;
diff --git a/Engine/Source/Runtime/AnimGraphRuntime/Public/CommonAnimationLibrary.h b/Engine/Source/Runtime/AnimGraphRuntime/Public/CommonAnimationLibrary.h
new file mode 100644
--- /dev/null
+++ b/Engine/Source/Runtime/AnimGraphRuntime/Public/CommonAnimationLibrary.h
@@ -0,0 +1,96 @@
+// Copyright 1998-2019 Epic Games, Inc. All Rights Reserved.
+
+#pragma once
+
+#include "CoreMinimal.h"
+#include "Curves/CurveFloat.h"
+#include "CommonAnimationTypes.h"
+
+/**
+ *	A library of the most common animation functions.
+ */
+namespace CommonAnimationLibrary
+{
+	/**
+	 *	This function perform easing on a float value using a variety of easing types.
+	 *	@param Value The float value to ease
+	 *	@param EasingType The easing function to use
+	 *	@param CustomCurve The curve to use if the easing type is "Custom"
+	 *	@param bFlip If set to true the easing is flipping around
+	 *	@param Weight The amount of easing to use against linear (0.0 to 1.0)
+	 */
+	ANIMGRAPHRUNTIME_API float ScalarEasing(
+		float Value, 
+		const FRuntimeFloatCurve& CustomCurve,
+		EEasingFuncType EasingType = EEasingFuncType::Linear,
+		bool bFlip = false,
+		float Weight = 1.f
+	);
+
+	/**
+	 *	This function performs retargeting of translation using an easing function.
+	 *	For this a range of motion needs to be defined given a direction and the constraints.
+	 *	@param Location The input location to be retargeted
+	 *	@param Source The transform to be used as the frame of reference
+	 *	@param Target The transform to be used as the target space
+	 *	@param Axis The direction to use for the range measurement. Defaults to (1.0, 0.0, 0.0)
+	 *	@param SourceMinimum The minimum of the source range
+	 *	@param SourceMaximum The maximum of the source range
+	 *	@param TargetMinimum The minimum of the target range
+	 *	@param TargetMaximum The maximum of the target range
+	 *	@param EasingType The type of easing to apply
+	 *	@param CustomCurve The curve to use if the easing type is "Custom"
+	 *	@param bFlipEasing If set to true the easing is flipping around
+	 *	@param EasingWeight The amount of easing to use against linear (0.0 to 1.0)
+	 */
+	ANIMGRAPHRUNTIME_API FVector RetargetSingleLocation(
+		FVector Location,
+		FTransform Source,
+		FTransform Target,
+		const FRuntimeFloatCurve& CustomCurve,
+		EEasingFuncType EasingType = EEasingFuncType::Linear,
+		bool bFlipEasing = false,
+		float EasingWeight = 1.f,
+		FVector Axis = FVector(1.f, 0.f, 0.f),
+		float SourceMinimum = -1.f,
+		float SourceMaximum = 1.f,
+		float TargetMinimum = -1.f,
+		float TargetMaximum = 1.f
+	);
+
+	/**
+	 *	This function performs retargeting of rotation using an easing function. 
+	 *	For this a range of motion needs to be defined as a euler angle, swing angle or twist
+	 *	@param Rotation The input rotation to be retargeted
+	 *	@param Source The transform to be used as the frame of reference
+	 *	@param Target The transform to be used as the target space
+	 *	@param CustomCurve The curve to use if the easing type is "Custom"
+	 *	@param bFlipEasing If set to true the easing is flipping around
+	 *	@param EasingWeight The amount of easing to use against linear (0.0 to 1.0)
+	 *	@param RotationComponent The component of the rotation to retarget
+	 *	@param TwistAxis The axis to use when extracting swing / twist rotations
+	 *	@param bUseAbsoluteAngle If set to true negative angles will be flipped to positive. This can be used to mirror the rotation.
+	 *	@param SourceMinimum The minimum of the source range in degrees
+	 *	@param SourceMaximum The maximum of the source range in degrees
+	 *	@param TargetMinimum The minimum of the target range in degrees
+	 *	@param TargetMaximum The maximum of the target range in degrees
+	 *	@param EasingType The type of easing to apply
+	 */
+	ANIMGRAPHRUNTIME_API FQuat RetargetSingleRotation(
+		FQuat Rotation,
+		FTransform Source,
+		FTransform Target,
+		const FRuntimeFloatCurve& CustomCurve,
+		EEasingFuncType EasingType = EEasingFuncType::Linear,
+		bool bFlipEasing = false,
+		float EasingWeight = 1.f,
+		ERotationComponent RotationComponent = ERotationComponent::SwingAngle,
+		FVector TwistAxis = FVector(1.f, 0.f, 0.f),
+		bool bUseAbsoluteAngle = false,
+		float SourceMinimum = 0.0f,
+		float SourceMaximum = 45.f,
+		float TargetMinimum = 0.f,
+		float TargetMaximum = 45.f
+	);
+};
+
diff --git a/Engine/Source/Runtime/AnimGraphRuntime/Public/CommonAnimationTypes.h b/Engine/Source/Runtime/AnimGraphRuntime/Public/CommonAnimationTypes.h
new file mode 100644
--- /dev/null
+++ b/Engine/Source/Runtime/AnimGraphRuntime/Public/CommonAnimationTypes.h
@@ -0,0 +1,156 @@
+// Copyright 1998-2019 Epic Games, Inc. All Rights Reserved.
+
+#pragma once
+
+#include "CoreMinimal.h"
+#include "Curves/CurveFloat.h"
+#include "CommonAnimationTypes.generated.h"
+
+/**
+ *	An easing type defining how to ease float values.
+ */
+UENUM()
+enum class EEasingFuncType : uint8
+{
+	// Linear easing (no change to the value)
+	Linear,
+	// Easing using a sinus function
+	Sinusoidal,
+	// Cubic version of the value (only in)
+	Cubic,
+	// Quadratic version of the value (in and out)
+	QuadraticInOut,
+	// Cubic version of the value (in and out)
+	CubicInOut,
+	// Easing using a cubic hermite function
+	HermiteCubic,
+	// Quartic version of the value (in and out)
+	QuarticInOut,
+	// Quintic version of the value (in and out)
+	QuinticInOut,
+	// Circular easing (only in)
+	CircularIn,
+	// Circular easing (only out)
+	CircularOut,
+	// Circular easing (in and out)
+	CircularInOut,
+	// Exponential easing (only in)
+	ExpIn,
+	// Exponential easing (only out)
+	ExpOut,
+	// Exponential easing (in and out)
+	ExpInOut,
+	// Custom - based on an optional Curve
+	CustomCurve
+};
+
+// A rotational component. This is used for retargeting, for example.
+UENUM()
+enum class ERotationComponent : uint8
+{
+	// Using the X component of the Euler rotation
+	EulerX,
+	// Using the Y component of the Euler rotation
+	EulerY,
+	// Using the Z component of the Euler rotation
+	EulerZ,
+	// Using the angle of the quaternion
+	QuaternionAngle,
+	// Using the angle of the swing quaternion
+	SwingAngle,
+	// Using the angle of the twist quaternion
+	TwistAngle
+};
+
+/**
+ *	The FRotationRetargetingInfo is used to provide all of the 
+ *	settings required to perform rotational retargeting on a single
+ *	transform.
+ */
+USTRUCT(BlueprintType)
+struct FRotationRetargetingInfo
+{
+	GENERATED_BODY()
+
+public:
+
+	/** Default constructor */
+	FRotationRetargetingInfo(bool bInEnabled = true)
+		: bEnabled(bInEnabled)
+		, Source(FTransform::Identity)
+		, Target(FTransform::Identity)
+		, RotationComponent(ERotationComponent::SwingAngle)
+		, TwistAxis(FVector(1.f, 0.f, 0.f))
+		, bUseAbsoluteAngle(false)
+		, SourceMinimum(0.f)
+		, SourceMaximum(45.f)
+		, TargetMinimum(0.f)
+		, TargetMaximum(45.f)
+		, EasingType(EEasingFuncType::Linear)
+		, CustomCurve(FRuntimeFloatCurve())
+		, bFlipEasing(false)
+		, EasingWeight(1.f)
+		, bClamp(false)
+	{
+	}
+
+	// Set to true this enables retargeting
+	UPROPERTY(EditAnywhere, Category = "FRotationRetargetingInfo")
+	bool bEnabled;
+
+	// The source transform of the frame of reference. The rotation is made relative to this space
+	UPROPERTY(EditAnywhere, Category = "FRotationRetargetingInfo")
+	FTransform Source;
+
+	// The target transform to project the rotation. In most cases this is the same as Source
+	UPROPERTY(EditAnywhere, Category = "FRotationRetargetingInfo")
+	FTransform Target;
+
+	// The rotation component to perform retargeting with
+	UPROPERTY(EditAnywhere, Category = "FRotationRetargetingInfo")
+	ERotationComponent RotationComponent;
+
+	// In case the rotation component is SwingAngle or TwistAngle this vector is used as the twist axis
+	UPROPERTY(EditAnywhere, Category = "FRotationRetargetingInfo")
+	FVector TwistAxis;
+
+	// If set to true the angle will be always positive, thus resulting in mirrored rotation both ways
+	UPROPERTY(EditAnywhere, Category = "FRotationRetargetingInfo")
+	bool bUseAbsoluteAngle;
+
+	// The minimum value of the source angle in degrees
+	UPROPERTY(EditAnywhere, Category = "FRotationRetargetingInfo", meta = (UIMin = "-90.0", UIMax = "90.0"))
+	float SourceMinimum;
+
+	// The maximum value of the source angle in degrees
+	UPROPERTY(EditAnywhere, Category = "FRotationRetargetingInfo", meta = (UIMin = "-90.0", UIMax = "90.0"))
+	float SourceMaximum;
+
+	// The minimum value of the target angle in degrees (can be the same as SourceMinimum)
+	UPROPERTY(EditAnywhere, Category = "FRotationRetargetingInfo", meta = (UIMin = "-90.0", UIMax = "90.0"))
+	float TargetMinimum;
+
+	// The target value of the target angle in degrees (can be the same as SourceMaximum)
+	UPROPERTY(EditAnywhere, Category = "FRotationRetargetingInfo", meta = (UIMin = "-90.0", UIMax = "90.0"))
+	float TargetMaximum;
+
+	// The easing to use - pick linear if you don't want to apply any easing
+	UPROPERTY(EditAnywhere, Category = "FRotationRetargetingInfo")
+	EEasingFuncType EasingType;
+
+	/** Custom curve mapping to apply if bApplyCustomCurve is true */
+	UPROPERTY(EditAnywhere, Category = "FRotationRetargetingInfo")
+	FRuntimeFloatCurve CustomCurve;
+
+	// If set to true the interpolation value for the easing will be flipped (1.0 - Value)
+	UPROPERTY(EditAnywhere, Category = "FRotationRetargetingInfo")
+	bool bFlipEasing;
+
+	// The amount of easing to apply (value should be 0.0 to 1.0)
+	UPROPERTY(EditAnywhere, Category = "FRotationRetargetingInfo", meta = (UIMin = "0.0", UIMax = "1.0"))
+	float EasingWeight;
+	
+	// If set to true the value for the easing will be clamped between 0.0 and 1.0
+	UPROPERTY(EditAnywhere, Category = "FRotationRetargetingInfo")
+	bool bClamp;
+};
diff --git a/Engine/Source/Runtime/AnimGraphRuntime/Public/KismetAnimationLibrary.h b/Engine/Source/Runtime/AnimGraphRuntime/Public/KismetAnimationLibrary.h
--- a/Engine/Source/Runtime/AnimGraphRuntime/Public/KismetAnimationLibrary.h
+++ b/Engine/Source/Runtime/AnimGraphRuntime/Public/KismetAnimationLibrary.h
@@ -5,45 +5,156 @@
 #include "CoreMinimal.h"
 #include "Kismet/BlueprintFunctionLibrary.h"
 #include "Components/SceneComponent.h"
+#include "Curves/CurveFloat.h"
 #include "KismetAnimationTypes.h"
 #include "KismetAnimationLibrary.generated.h"
 
 class USkeletalMeshComponent;
 
+/**
+ *	A library of the most common animation blueprint functions.
+ */
 UCLASS(meta=(ScriptName="AnimGraphLibrary"))
 class ANIMGRAPHRUNTIME_API UKismetAnimationLibrary : public UBlueprintFunctionLibrary
 {
 	GENERATED_UCLASS_BODY()
 
+	/**
+	 *	Computes the transform for two bones using inverse kinematics.
+	 *
+	 *	@param RootPos The input root position of the two bone chain
+	 *	@param JointPos The input center (elbow) position of the two bone chain
+	 *	@param EndPos The input end (wrist) position of the two bone chain
+	 *	@param JointTarget The IK target for the write to reach
+	 *	@param Effector The position of the target effector for the IK Chain.
+	 *	@param OutJointPos The resulting position for the center (elbow)
+	 *	@param OutEndPos The resulting position for the end (wrist)
+	 *	@param bAllowStretching If set to true the bones are allowed to stretch
+	 *	@param StartStretchRatio The ratio at which the bones should start to stretch. The higher the value, the later the stretching wil start.
+	 *	@param MaxStretchScale The maximum multiplier for the stretch to reach.
+	 */
 	UFUNCTION(BlueprintPure, Category = "Utilities|Animation", meta = (DisplayName = "Two Bone IK Function", ScriptName = "TwoBoneIK", bAllowStretching = "false", StartStretchRatio = "1.0", MaxStretchScale = "1.2"))
-	static void K2_TwoBoneIK(const FVector& RootPos, const FVector& JointPos, const FVector& EndPos, const FVector& JointTarget, const FVector& Effector, FVector& OutJointPos, FVector& OutEndPos, bool bAllowStretching, float StartStretchRatio, float MaxStretchScale);
+	static void K2_TwoBoneIK(
+		const FVector& RootPos, 
+		const FVector& JointPos, 
+		const FVector& EndPos, 
+		const FVector& JointTarget, 
+		const FVector& Effector, 
+		FVector& OutJointPos, 
+		FVector& OutEndPos, 
+		bool bAllowStretching, 
+		float StartStretchRatio,
+		float MaxStretchScale
+	);
 
+	/**
+	 *	Computes the transform which is "looking" at target position with a local axis.
+	 *
+	 *	@param CurrentTransform The input transform to modify
+	 *	@param TargetPosition The position this transform should look at
+	 *	@param LookAtVector The local vector to align with the target
+	 *	@param bUseUpVector If set to true the lookat will also perform a twist rotation
+	 *	@param UpVector The position to use for the upvector target (only used is bUseUpVector is turned on)
+	 *	@param ClampConeInDegree A limit for only allowing the lookat to rotate as much as defined by the float value
+	 */
 	UFUNCTION(BlueprintPure, Category = "Utilities|Animation", meta = (DisplayName = "Look At Function", ScriptName = "LookAt", bUseUpVector = "false"))
-	static FTransform K2_LookAt(const FTransform& CurrentTransform, const FVector& TargetPosition, FVector LookAtVector, bool bUseUpVector, FVector UpVector, float ClampConeInDegree);
+	static FTransform K2_LookAt(
+		const FTransform& CurrentTransform, 
+		const FVector& TargetPosition, 
+		FVector LookAtVector, 
+		bool bUseUpVector, 
+		FVector UpVector, 
+		float ClampConeInDegree
+	);
 
+	/**
+	 *	Computes the distance between two bones / sockets and can remap the range.
+	 *
+	 *	@param Component The skeletal component to look for the sockets / bones within
+	 *	@param SocketOrBoneNameA The name of the first socket / bone
+	 *	@param SocketSpaceA The space for the first socket / bone
+	 *	@param SocketOrBoneNameB The name of the second socket / bone
+	 *	@param SocketSpaceB The space for the second socket / bone
+	 *	@param bRemapRange If set to true, the distance will be remapped using the range parameters
+	 *	@param InRangeMin The minimum for the input range (commonly == 0.0)
+	 *	@param InRangeMax The maximum for the input range (the max expected distance)
+	 *	@param OutRangeMin The minimum for the output range (commonly == 0.0)
+	 *	@param OutRangeMax The maximum for the output range (commonly == 1.0)
+	 */
 	UFUNCTION(BlueprintPure, Category = "Utilities|Animation", meta = (DisplayName = "Get Distance Between Two Sockets", ScriptName = "DistanceBetweenSockets", bRemapRange= "false"))
-	static float K2_DistanceBetweenTwoSocketsAndMapRange(const USkeletalMeshComponent* Component, const FName SocketOrBoneNameA, ERelativeTransformSpace SocketSpaceA, const FName SocketOrBoneNameB, ERelativeTransformSpace SocketSpaceB, bool bRemapRange, float InRangeMin, float InRangeMax, float OutRangeMin, float OutRangeMax);
-
+	static float K2_DistanceBetweenTwoSocketsAndMapRange(
+		const USkeletalMeshComponent* Component, 
+		const FName SocketOrBoneNameA, 
+		ERelativeTransformSpace SocketSpaceA, 
+		const FName SocketOrBoneNameB, 
+		ERelativeTransformSpace SocketSpaceB, 
+		bool bRemapRange, 
+		float InRangeMin,
+		float InRangeMax,
+		float OutRangeMin,
+		float OutRangeMax);
+
+	/**
+	 *	Computes the direction between two bones / sockets.
+	 *
+	 *	@param Component The skeletal component to look for the sockets / bones within
+	 *	@param SocketOrBoneNameFrom The name of the first socket / bone
+	 *	@param SocketOrBoneNameTo The name of the second socket / bone
+	 */
 	UFUNCTION(BlueprintPure, Category = "Utilities|Animation", meta = (DisplayName = "Get Direction Between Sockets", ScriptName = "DirectionBetweenSockets"))
 	static FVector K2_DirectionBetweenSockets(const USkeletalMeshComponent* Component, const FName SocketOrBoneNameFrom, const FName SocketOrBoneNameTo);
 
-	/* This function creates perlin noise from input X, Y, Z, and then range map to RangeOut, and out put to OutX, OutY, OutZ */
+	/** 
+	 *	This function creates perlin noise from input X, Y, Z, and then range map to RangeOut, and out put to OutX, OutY, OutZ
+	 *
+	 *	@param X The x component for the input position of the noise
+	 *	@param Y The y component for the input position of the noise
+	 *	@param Z The z component for the input position of the noise
+	 *	@param RangeOutMinX The minimum for the output range for the x component
+	 *	@param RangeOutMaxX The maximum for the output range for the x component
+	 *	@param RangeOutMinY The minimum for the output range for the y component
+	 *	@param RangeOutMaxY The maximum for the output range for the y component
+	 *	@param RangeOutMinZ The minimum for the output range for the z component
+	 *	@param RangeOutMaxZ The maximum for the output range for the z component
+	 */
 	UFUNCTION(BlueprintPure, Category = "Utilities|Animation", meta = (DisplayName = "Make Perlin Noise Vector and Remap", ScriptName = "MakeVectorFromPerlinNoise", RangeOutMinX = "-1.f", RangeOutMaxX = "1.f", RangeOutMinY = "-1.f", RangeOutMaxY = "1.f", RangeOutMinZ = "-1.f", RangeOutMaxZ = "1.f"))
-	static FVector K2_MakePerlinNoiseVectorAndRemap(float X, float Y, float Z, float RangeOutMinX, float RangeOutMaxX, float RangeOutMinY, float RangeOutMaxY, float RangeOutMinZ, float RangeOutMaxZ);
+	static FVector K2_MakePerlinNoiseVectorAndRemap(
+		float X,
+		float Y,
+		float Z,
+		float RangeOutMinX,
+		float RangeOutMaxX,
+		float RangeOutMinY,
+		float RangeOutMaxY,
+		float RangeOutMinZ,
+		float RangeOutMaxZ
+	);
 
+	/**
+	 *	This function creates perlin noise for a single float and then range map to RangeOut
+	 *
+	 *	@param Value The input value for the noise function
+	 *	@param RangeOutMin The minimum for the output range
+	 *	@param RangeOutMax The maximum for the output range
+	 */
 	UFUNCTION(BlueprintPure, Category = "Utilities|Animation", meta = (DisplayName = "Make Perlin Noise and Remap", ScriptName = "MakeFloatFromPerlinNoise", RangeOutMinX = "-1.f", RangeOutMaxX = "1.f", RangeOutMinY = "-1.f", RangeOutMaxY = "1.f", RangeOutMinZ = "-1.f", RangeOutMaxZ = "1.f"))
-	static float K2_MakePerlinNoiseAndRemap(float Value, float RangeOutMin, float RangeOutMax);
+	static float K2_MakePerlinNoiseAndRemap(
+		float Value,
+		float RangeOutMin,
+		float RangeOutMax
+	);
 
 	/** 
-		This function calculates the velocity of a position changing over time.
-		You need to hook up a valid PositionHistory variable to this for storage.
-
-		@param Position The position to track over time.
-		@param History The history to use for storage.
-		@param NumberOfSamples The number of samples to use for the history. The higher the number of samples - the smoother the velocity changes.
-		@param VelocityMin The minimum velocity to use for normalization (if both min and max are set to 0, normalization is turned off)
-		@param VelocityMax The maximum velocity to use for normalization (if both min and max are set to 0, normalization is turned off)
-	*/
+	 *	This function calculates the velocity of a position changing over time.
+	 *	You need to hook up a valid PositionHistory variable to this for storage.
+	 *
+	 *	@param DeltaSeconds The time passed in seconds
+	 *	@param Position The position to track over time.
+	 *	@param History The history to use for storage.
+	 *	@param NumberOfSamples The number of samples to use for the history. The higher the number of samples - the smoother the velocity changes.
+	 *	@param VelocityMin The minimum velocity to use for normalization (if both min and max are set to 0, normalization is turned off)
+	 *	@param VelocityMax The maximum velocity to use for normalization (if both min and max are set to 0, normalization is turned off)
+	 */
 	UFUNCTION(BlueprintCallable, Category = "Utilities|Animation", meta = (DisplayName = "Calculate Velocity From Position History", ScriptName = "CalculateVelocityFromPositionHistory", NumberOfSamples = "16", VelocityMin = "0.f", VelocityMax = "128.f"))
 	static float K2_CalculateVelocityFromPositionHistory(
 		float DeltaSeconds,
@@ -54,25 +165,27 @@ class ANIMGRAPHRUNTIME_API UKismetAnimationLibrary : public UBlueprintFunctionLi
 		float VelocityMax
 	);
 
-	UFUNCTION(BlueprintPure, Category = "Utilities|Animation", meta = (DisplayName = "Scalar Easing", ScriptName = "ScalarEasing", Value = "0.f"))
-	static float K2_ScalarEasing(float Value, EEasingFuncType EasingType);
-
 	/** 
-		This function calculates the velocity of an offset position on a bone / socket over time.
-		The bone's / socket's motion can be expressed within a reference frame (another bone / socket). 
-		You need to hook up a valid PositionHistory variable to this for storage.
-
-		@param SocketOrBoneName The name of the bone / socket to track.
-		@param ReferenceSocketOrBone The name of the bone / socket to use as a frame of reference (or None if no frame of reference == world space).
-		@param OffsetInBoneSpace The relative position in the space of the bone / socket to track over time.
-		@param History The history to use for storage.
-		@param NumberOfSamples The number of samples to use for the history. The higher the number of samples - the smoother the velocity changes.
-		@param VelocityMin The minimum velocity to use for normalization (if both min and max are set to 0, normalization is turned off)
-		@param VelocityMax The maximum velocity to use for normalization (if both min and max are set to 0, normalization is turned off)
-	*/
+	 *	This function calculates the velocity of an offset position on a bone / socket over time.
+	 *	The bone's / socket's motion can be expressed within a reference frame (another bone / socket). 
+	 *	You need to hook up a valid PositionHistory variable to this for storage.
+	 *
+	 *	@param DeltaSeconds The time passed in seconds
+	 *	@param Component The skeletal component to look for the bones / sockets
+	 *	@param SocketOrBoneName The name of the bone / socket to track.
+	 *	@param ReferenceSocketOrBone The name of the bone / socket to use as a frame of reference (or None if no frame of reference == world space).
+	 *	@param SocketSpace The space to use for the two sockets / bones
+	 *	@param OffsetInBoneSpace The relative position in the space of the bone / socket to track over time.
+	 *	@param History The history to use for storage.
+	 *	@param NumberOfSamples The number of samples to use for the history. The higher the number of samples - the smoother the velocity changes.
+	 *	@param VelocityMin The minimum velocity to use for normalization (if both min and max are set to 0, normalization is turned off)
+	 *	@param VelocityMax The maximum velocity to use for normalization (if both min and max are set to 0, normalization is turned off)
+	 *	@param EasingType The easing function to use
+	 *	@param CustomCurve The curve to use if the easing type is "Custom"
+	 */
 	UFUNCTION(BlueprintCallable, Category = "Utilities|Animation", meta = (DisplayName = "Calculate Velocity From Sockets", ScriptName = "CalculateVelocityFromSockets", NumberOfSamples = "16", VelocityMin = "0.f", VelocityMax = "128.f"))
 	static float K2_CalculateVelocityFromSockets(
-		float DeltaSeconds, 
+		float DeltaSeconds,
 		USkeletalMeshComponent * Component,
 		const FName SocketOrBoneName,
 		const FName ReferenceSocketOrBone,
@@ -80,9 +193,11 @@ class ANIMGRAPHRUNTIME_API UKismetAnimationLibrary : public UBlueprintFunctionLi
 		FVector OffsetInBoneSpace, 
 		UPARAM(ref) FPositionHistory& History, 
 		int32 NumberOfSamples,
-		float VelocityMin, 
+		float VelocityMin,
 		float VelocityMax,
-		EEasingFuncType EasingType
-	);
+		EEasingFuncType EasingType,
+		const FRuntimeFloatCurve& CustomCurve
+		);
+
 };
 
diff --git a/Engine/Source/Runtime/AnimGraphRuntime/Public/KismetAnimationTypes.h b/Engine/Source/Runtime/AnimGraphRuntime/Public/KismetAnimationTypes.h
--- a/Engine/Source/Runtime/AnimGraphRuntime/Public/KismetAnimationTypes.h
+++ b/Engine/Source/Runtime/AnimGraphRuntime/Public/KismetAnimationTypes.h
@@ -3,28 +3,16 @@
 #pragma once
 
 #include "CoreMinimal.h"
+#include "CommonAnimationTypes.h"
 #include "KismetAnimationTypes.generated.h"
 
-// The kind of easing function to use
-UENUM()
-enum class EEasingFuncType : uint8
-{
-	Linear,
-	Sinusoidal,
-	Cubic,
-	QuadraticInOut,
-	CubicInOut,
-	HermiteCubic,
-	QuarticInOut,
-	QuinticInOut,
-	CircularIn,
-	CircularOut,
-	CircularInOut,
-	ExpIn,
-	ExpOut,
-	ExpInOut
-};
-
+/**
+ *	An easing type defining how to ease float values.
+ *	The FPositionHistory is a container to record position changes
+ *	over time. This is used to calculate velocity of a bone, for example.
+ *	The FPositionArray also tracks the last index used to allow for 
+ *	reuse of entries (instead of appending to the array all the time).
+ */
 USTRUCT(BlueprintType)
 struct FPositionHistory
 {
@@ -35,15 +23,19 @@ public:
 	/** Default constructor */
 	FPositionHistory()
 		: Positions(TArray<FVector>())
+		, Range(0.f)
 		, Velocities(TArray<float>())
 		, LastIndex(0)
 	{}
 
 	/** The recorded positions */
-	UPROPERTY(EditAnywhere, Category = "FPositionHistory", meta = (MetaClass = "AdvancedCopyCustomization"))
+	UPROPERTY(EditAnywhere, Category = "FPositionHistory")
 	TArray<FVector> Positions;
 
+	/** The range for this particular history */
+	UPROPERTY(EditAnywhere, Category = "FPositionHistory", meta = (UIMin = "0.0", UIMax = "1.0"))
+	float Range;
+
 	TArray<float> Velocities;
 	uint32 LastIndex; /// The last index used to store a position
 };
-
