Fix PSD cache not getting cleared and produced incorrect poses at the end

diff --git a/Engine/Source/Runtime/AnimGraphRuntime/Private/AnimNodes/AnimNode_PoseDriver.cpp b/Engine/Source/Runtime/AnimGraphRuntime/Private/AnimNodes/AnimNode_PoseDriver.cpp
--- a/Engine/Source/Runtime/AnimGraphRuntime/Private/AnimNodes/AnimNode_PoseDriver.cpp
+++ b/Engine/Source/Runtime/AnimGraphRuntime/Private/AnimNodes/AnimNode_PoseDriver.cpp
@@ -309,6 +309,12 @@ void FAnimNode_PoseDriver::Evaluate_AnyThread(FPoseContext& Output)
 		{
 			FPoseContext CurrentPose(Output);
 
+			// clear the value before setting it. 
+			for (int32 PoseIndex = 0; PoseIndex < PoseExtractContext.PoseCurves.Num(); ++PoseIndex)
+			{
+				PoseExtractContext.PoseCurves[PoseIndex].Value = 0.f;
+			}
+
 			// Then fill in weight for any driven poses
 			for (const FRBFOutputWeight& Weight : OutputWeights)
 			{
