AnimGraph: Add profiling blueprint functions to measure anim blueprints.

diff --git a/Engine/Source/Runtime/AnimGraphRuntime/Private/KismetAnimationLibrary.cpp b/Engine/Source/Runtime/AnimGraphRuntime/Private/KismetAnimationLibrary.cpp
--- a/Engine/Source/Runtime/AnimGraphRuntime/Private/KismetAnimationLibrary.cpp
+++ b/Engine/Source/Runtime/AnimGraphRuntime/Private/KismetAnimationLibrary.cpp
@@ -215,5 +215,44 @@ float UKismetAnimationLibrary::K2_CalculateVelocityFromSockets(
 	return VelocityMin;
 }
 
+
+TArray<UKismetAnimationLibrary::FK2ProfilingTimer> UKismetAnimationLibrary::sProfilingTimers;
+
+void UKismetAnimationLibrary::K2_StartProfilingTimer()
+{
+	FK2ProfilingTimer Timer;
+	Timer.LastTime = FPlatformTime::Seconds() * 1000.0;
+	Timer.AccummulatedTime = 0.0;
+	sProfilingTimers.Add(Timer);
+}
+
+float UKismetAnimationLibrary::K2_EndProfilingTimer(bool bLog, const FString& LogPrefix)
+{
+	if (sProfilingTimers.Num() == 0)
+	{
+		UE_LOG(LogAnimation, Warning, TEXT("Unbalanced use of Start & End Profiling Timer nodes."));
+		return 0.f;
+	}
+
+	FK2ProfilingTimer Timer = sProfilingTimers.Pop();
+	double CurrentTimer = FPlatformTime::Seconds() * 1000.0;
+	Timer.AccummulatedTime = CurrentTimer - Timer.LastTime;
+
+	float Delta = (float)Timer.AccummulatedTime;
+	if (bLog)
+	{
+		if (LogPrefix.IsEmpty())
+		{
+			UE_LOG(LogAnimation, Warning, TEXT("%.03f ms"), Delta);
+		}
+		else
+		{
+			UE_LOG(LogAnimation, Warning, TEXT("[%s] %.03f ms"), *LogPrefix, Delta);
+		}
+	}
+
+	return Delta;
+}
+
 #undef LOCTEXT_NAMESPACE
 
diff --git a/Engine/Source/Runtime/AnimGraphRuntime/Public/KismetAnimationLibrary.h b/Engine/Source/Runtime/AnimGraphRuntime/Public/KismetAnimationLibrary.h
--- a/Engine/Source/Runtime/AnimGraphRuntime/Public/KismetAnimationLibrary.h
+++ b/Engine/Source/Runtime/AnimGraphRuntime/Public/KismetAnimationLibrary.h
@@ -199,5 +199,29 @@ class ANIMGRAPHRUNTIME_API UKismetAnimationLibrary : public UBlueprintFunctionLi
 		const FRuntimeFloatCurve& CustomCurve
 		);
 
+	/**
+	 *	This function starts measuring the time for a profiling bracket
+	 */
+	UFUNCTION(BlueprintCallable, Category = "Utilities|Animation", meta = (DisplayName = "Start Profiling Timer"))
+	static void K2_StartProfilingTimer();
+
+	/**
+	 *	This function ends measuring a profiling bracket and optionally logs the result
+	 *
+	 *  @param bLog If set to true the result is logged to the OutputLog
+	 *  @param LogPrefix A prefix to use for the log
+	 *  @result The time spent in milliseconds
+	 */
+	UFUNCTION(BlueprintCallable, Category = "Utilities|Animation", meta = (DisplayName = "End Profiling Timer"))
+	static float K2_EndProfilingTimer(bool bLog = true, const FString& LogPrefix = "");
+
+private:
+
+	struct FK2ProfilingTimer
+	{
+		double LastTime;
+		double AccummulatedTime;
+	};
+	static TArray<FK2ProfilingTimer> sProfilingTimers;
 };
 
