Control Rig MVC Refactoring

diff --git a/Engine/Source/Editor/KismetCompiler/Private/AnimBlueprintCompiler.cpp b/Engine/Source/Editor/KismetCompiler/Private/AnimBlueprintCompiler.cpp
--- a/Engine/Source/Editor/KismetCompiler/Private/AnimBlueprintCompiler.cpp
+++ b/Engine/Source/Editor/KismetCompiler/Private/AnimBlueprintCompiler.cpp
@@ -560,14 +560,14 @@ void FAnimBlueprintCompilerContext::ProcessAnimationNode(UAnimGraphNode_Base* Vi
 			int32 SourceArrayIndex = INDEX_NONE;
 
 			// We have special handling below if we're targeting a subinstance instead of our own instance properties
-			UAnimGraphNode_SubInstance* SubInstanceNode = Cast<UAnimGraphNode_SubInstance>(VisualAnimNode);
+			UAnimGraphNode_CustomProperty* CustomPropertyNode = Cast<UAnimGraphNode_CustomProperty>(VisualAnimNode);
 
 			// Does this pin have an associated evaluation handler?
-			if(SubInstanceNode)
+			if(CustomPropertyNode)
 			{
 				// Subinstance nodes use instance properties not node properties as they aren't UObjects
 				// and we can't store non-native properties there
-				SubInstanceNode->GetInstancePinProperty(NewAnimBlueprintClass, SourcePin, SourcePinProperty);
+				CustomPropertyNode->GetInstancePinProperty(NewAnimBlueprintClass, SourcePin, SourcePinProperty);
 			}
 			else
 			{
@@ -601,7 +601,7 @@ void FAnimBlueprintCompilerContext::ProcessAnimationNode(UAnimGraphNode_Base* Vi
 					EvalHandler.NodeVariableProperty = NewProperty;
 					EvalHandler.RegisterPin(SourcePin, SourcePinProperty, SourceArrayIndex);
 
-					if(SubInstanceNode)
+					if(CustomPropertyNode)
 					{
 						EvalHandler.bServicesInstanceProperties = true;
 
diff --git a/Engine/Source/Runtime/AnimGraphRuntime/Private/BoneControllers/AnimNode_CCDIK.cpp b/Engine/Source/Runtime/AnimGraphRuntime/Private/BoneControllers/AnimNode_CCDIK.cpp
--- a/Engine/Source/Runtime/AnimGraphRuntime/Private/BoneControllers/AnimNode_CCDIK.cpp
+++ b/Engine/Source/Runtime/AnimGraphRuntime/Private/BoneControllers/AnimNode_CCDIK.cpp
@@ -80,7 +80,7 @@ void FAnimNode_CCDIK::EvaluateSkeletalControl_AnyThread(FComponentSpacePoseConte
 		const FTransform& BoneCSTransform = Output.Pose.GetComponentSpaceTransform(RootBoneIndex);
 
 		OutBoneTransforms[0] = FBoneTransform(RootBoneIndex, BoneCSTransform);
-		Chain.Add(CCDIKChainLink(BoneCSTransform, LocalTransform, RootBoneIndex, 0));
+		Chain.Add(CCDIKChainLink(BoneCSTransform, LocalTransform, 0));
 	}
 
 	// Go through remaining transforms
@@ -99,7 +99,7 @@ void FAnimNode_CCDIK::EvaluateSkeletalControl_AnyThread(FComponentSpacePoseConte
 
 		if (!FMath::IsNearlyZero(BoneLength))
 		{
-			Chain.Add(CCDIKChainLink(BoneCSTransform, LocalTransform, BoneIndex, TransformIndex));
+			Chain.Add(CCDIKChainLink(BoneCSTransform, LocalTransform, TransformIndex));
 		}
 		else
 		{
@@ -110,53 +110,14 @@ void FAnimNode_CCDIK::EvaluateSkeletalControl_AnyThread(FComponentSpacePoseConte
 		}
 	}
 
-	bool bBoneLocationUpdated = false;
-	int32 const NumChainLinks = BoneIndices.Num();
-
-	// iterate
-	{
-		int32 const TipBoneLinkIndex = NumChainLinks - 1;
-
-		// @todo optimize locally if no update, stop?
-		bool bLocalUpdated = false;
-		// check how far
-		const FVector TargetPos = CSEffectorLocation;
-		FVector TipPos = Chain[TipBoneLinkIndex].Transform.GetLocation();
-		float Distance = FVector::Dist(TargetPos, TipPos);
-		int32 IterationCount = 0;
-		while ((Distance > Precision) && (IterationCount++ < MaxIterations))
-		{
-			// iterate from tip to root
-			if (bStartFromTail)
-			{
-				for (int32 LinkIndex = TipBoneLinkIndex - 1; LinkIndex > 0; --LinkIndex)
-				{
-					bLocalUpdated |= UpdateChainLink(Chain, LinkIndex, TargetPos);
-				}
-			}
-			else
-			{
-				for (int32 LinkIndex = 1; LinkIndex < TipBoneLinkIndex; ++LinkIndex)
-				{
-					bLocalUpdated |= UpdateChainLink(Chain, LinkIndex, TargetPos);
-				}
-			}
-
-			Distance = FVector::Dist(Chain[TipBoneLinkIndex].Transform.GetLocation(), CSEffectorLocation);
-
-			bBoneLocationUpdated |= bLocalUpdated;
-
-			// no more update in this iteration
-			if (!bLocalUpdated)
-			{
-				break;
-			}
-		}
-	}
+	// solve
+	bool bBoneLocationUpdated = AnimationCore::SolveCCDIK(Chain, CSEffectorLocation, Precision, MaxIterations, bStartFromTail, bEnableRotationLimit, RotationLimitPerJoints);
 
 	// If we moved some bones, update bone transforms.
 	if (bBoneLocationUpdated)
 	{
+		int32 NumChainLinks = Chain.Num();
+
 		// First step: update bone transform positions from chain links.
 		for (int32 LinkIndex = 0; LinkIndex < NumChainLinks; LinkIndex++)
 		{
@@ -183,84 +144,6 @@ void FAnimNode_CCDIK::EvaluateSkeletalControl_AnyThread(FComponentSpacePoseConte
 	}
 }
 
-bool FAnimNode_CCDIK::UpdateChainLink(TArray<CCDIKChainLink>& Chain, int32 LinkIndex, const FVector& TargetPos) const
-{
-	int32 const TipBoneLinkIndex = Chain.Num() - 1;
-
-	ensure(Chain.IsValidIndex(TipBoneLinkIndex));
-	CCDIKChainLink& CurrentLink = Chain[LinkIndex];
-
-	// update new tip pos
-	FVector TipPos = Chain[TipBoneLinkIndex].Transform.GetLocation();
-
-	FTransform& CurrentLinkTransform = CurrentLink.Transform;
-	FVector ToEnd = TipPos - CurrentLinkTransform.GetLocation();
-	FVector ToTarget = TargetPos - CurrentLinkTransform.GetLocation();
-
-	ToEnd.Normalize();
-	ToTarget.Normalize();
-
-	float RotationLimitPerJointInRadian = FMath::DegreesToRadians(RotationLimitPerJoints[LinkIndex]);
-	float Angle = FMath::ClampAngle(FMath::Acos(FVector::DotProduct(ToEnd, ToTarget)), -RotationLimitPerJointInRadian, RotationLimitPerJointInRadian);
-	bool bCanRotate = (FMath::Abs(Angle) > KINDA_SMALL_NUMBER) && (!bEnableRotationLimit || RotationLimitPerJointInRadian > CurrentLink.CurrentAngleDelta);
-	if (bCanRotate)
-	{
-		// check rotation limit first, if fails, just abort
-		if (bEnableRotationLimit)
-		{
-			if (RotationLimitPerJointInRadian < CurrentLink.CurrentAngleDelta + Angle)
-			{
-				Angle = RotationLimitPerJointInRadian - CurrentLink.CurrentAngleDelta;
-				if (Angle <= KINDA_SMALL_NUMBER)
-				{
-					return false;
-				}
-			}
-
-			CurrentLink.CurrentAngleDelta += Angle;
-		}
-
-		// continue with rotating toward to target
-		FVector RotationAxis = FVector::CrossProduct(ToEnd, ToTarget);
-		if (RotationAxis.SizeSquared() > 0.f)
-		{
-			RotationAxis.Normalize();
-			// Delta Rotation is the rotation to target
-			FQuat DeltaRotation(RotationAxis, Angle);
-
-			FQuat NewRotation = DeltaRotation * CurrentLinkTransform.GetRotation();
-			NewRotation.Normalize();
-			CurrentLinkTransform.SetRotation(NewRotation);
-
-			// if I have parent, make sure to refresh local transform since my current transform has changed
-			if (LinkIndex > 0)
-			{
-				CCDIKChainLink const & Parent = Chain[LinkIndex - 1];
-				CurrentLink.LocalTransform = CurrentLinkTransform.GetRelativeTransform(Parent.Transform);
-				CurrentLink.LocalTransform.NormalizeRotation();
-			}
-
-			// now update all my children to have proper transform
-			FTransform CurrentParentTransform = CurrentLinkTransform;
-
-			// now update all chain
-			for (int32 ChildLinkIndex = LinkIndex + 1; ChildLinkIndex <= TipBoneLinkIndex; ++ChildLinkIndex)
-			{
-				CCDIKChainLink& ChildIterLink = Chain[ChildLinkIndex];
-				FCompactPoseBoneIndex ChildBoneIndex = ChildIterLink.BoneIndex;
-				const FTransform LocalTransform = ChildIterLink.LocalTransform;
-				ChildIterLink.Transform = LocalTransform * CurrentParentTransform;
-				ChildIterLink.Transform.NormalizeRotation();
-				CurrentParentTransform = ChildIterLink.Transform;
-			}
-
-			return true;
-		}
-	}
-
-	return false;
-}
-
 bool FAnimNode_CCDIK::IsValidToEvaluate(const USkeleton* Skeleton, const FBoneContainer& RequiredBones)
 {
 	if (EffectorLocationSpace == BCS_ParentBoneSpace || EffectorLocationSpace == BCS_BoneSpace)
diff --git a/Engine/Source/Runtime/AnimGraphRuntime/Public/BoneControllers/AnimNode_CCDIK.h b/Engine/Source/Runtime/AnimGraphRuntime/Public/BoneControllers/AnimNode_CCDIK.h
--- a/Engine/Source/Runtime/AnimGraphRuntime/Public/BoneControllers/AnimNode_CCDIK.h
+++ b/Engine/Source/Runtime/AnimGraphRuntime/Public/BoneControllers/AnimNode_CCDIK.h
@@ -3,51 +3,12 @@
 #pragma once
 
 #include "AnimNode_SkeletalControlBase.h"
+#include "CCDIK.h"
 #include "AnimNode_CCDIK.generated.h"
 
 /**
 *	Controller which implements the CCDIK IK approximation algorithm
 */
-
-/** Transient structure for CCDIK node evaluation */
-struct CCDIKChainLink
-{
-public:
-	/** Transform of bone in component space. */
-	FTransform Transform;
-
-	/** Transform of bone in local space. This is mutable as their component space changes or parents*/
-	FTransform LocalTransform;
-
-	/** Bone Index in SkeletalMesh */
-	FCompactPoseBoneIndex BoneIndex;
-
-	/** Transform Index that this control will output */
-	int32 TransformIndex;
-
-	/** Child bones which are overlapping this bone. 
-	 * They have a zero length distance, so they will inherit this bone's transformation. */
-	TArray<int32> ChildZeroLengthTransformIndices;
-
-	float CurrentAngleDelta;
-
-	CCDIKChainLink()
-		: BoneIndex(INDEX_NONE)
-		, TransformIndex(INDEX_NONE)
-		, CurrentAngleDelta(0.f)
-	{
-	}
-
-	CCDIKChainLink(const FTransform& InTransform, const FTransform& InLocalTransform,  const FCompactPoseBoneIndex& InBoneIndex, const int32& InTransformIndex)
-		: Transform(InTransform)
-		, LocalTransform(InLocalTransform)
-		, BoneIndex(InBoneIndex)
-		, TransformIndex(InTransformIndex)
-		, CurrentAngleDelta(0.f)
-	{
-	}
-};
-
 USTRUCT()
 struct ANIMGRAPHRUNTIME_API FAnimNode_CCDIK : public FAnimNode_SkeletalControlBase
 {
@@ -116,8 +77,6 @@ private:
 
 	static FTransform GetTargetTransform(const FTransform& InComponentTransform, FCSPose<FCompactPose>& MeshBases, FBoneSocketTarget& InTarget, EBoneControlSpace Space, const FVector& InOffset);
 
-	// return true if updated
-	bool UpdateChainLink(TArray<CCDIKChainLink>& Chain, int32 LinkIndex, const FVector& TargetPos) const;
 public:
 #if WITH_EDITOR
 #if !(UE_BUILD_SHIPPING || UE_BUILD_TEST)
diff --git a/Engine/Source/Runtime/AnimationCore/Public/CCDIK.h b/Engine/Source/Runtime/AnimationCore/Public/CCDIK.h
new file mode 100644
--- /dev/null
+++ b/Engine/Source/Runtime/AnimationCore/Public/CCDIK.h
@@ -0,0 +1,46 @@
+// Copyright 1998-2019 Epic Games, Inc. All Rights Reserved.
+
+#pragma once
+
+#include "CoreMinimal.h"
+#include "UObject/ObjectMacros.h"
+#include "BoneIndices.h"
+
+/** Transient structure for CCDIK node evaluation */
+struct CCDIKChainLink
+{
+public:
+	/** Transform of bone in component space. */
+	FTransform Transform;
+
+	/** Transform of bone in local space. This is mutable as their component space changes or parents*/
+	FTransform LocalTransform;
+
+	/** Transform Index that this control will output */
+	int32 TransformIndex;
+
+	/** Child bones which are overlapping this bone. 
+	 * They have a zero length distance, so they will inherit this bone's transformation. */
+	TArray<int32> ChildZeroLengthTransformIndices;
+
+	float CurrentAngleDelta;
+
+	CCDIKChainLink()
+		: TransformIndex(INDEX_NONE)
+		, CurrentAngleDelta(0.f)
+	{
+	}
+
+	CCDIKChainLink(const FTransform& InTransform, const FTransform& InLocalTransform, const int32& InTransformIndex)
+		: Transform(InTransform)
+		, LocalTransform(InLocalTransform)
+		, TransformIndex(InTransformIndex)
+		, CurrentAngleDelta(0.f)
+	{
+	}
+};
+
+namespace AnimationCore
+{
+	ANIMATIONCORE_API bool SolveCCDIK(TArray<CCDIKChainLink>& InOutChain, const FVector& TargetPosition, float Precision, int32 MaxIteration, bool bStartFromTail, bool bEnableRotationLimit, const TArray<float>& RotationLimitPerJoints);
+};
diff --git a/Engine/Source/Runtime/Engine/Classes/Animation/AnimNode_CustomProperty.h b/Engine/Source/Runtime/Engine/Classes/Animation/AnimNode_CustomProperty.h
--- a/Engine/Source/Runtime/Engine/Classes/Animation/AnimNode_CustomProperty.h
+++ b/Engine/Source/Runtime/Engine/Classes/Animation/AnimNode_CustomProperty.h
@@ -51,6 +51,21 @@ protected:
 	UPROPERTY(Transient)
 	TArray<UProperty*> DestProperties;
 
+#if WITH_EDITOR
+	bool bReinitializeProperties;
+#endif // WITH_EDITOR
+
+	virtual bool HasPreUpdate() const override
+	{
+#if WITH_EDITOR
+		return true;
+#else
+		return false;
+#endif
+	}
+
+	virtual void PreUpdate(const UAnimInstance* InAnimInstance) override;
+
 	/* Set Target Instance */
 	void SetTargetInstance(UObject* InInstance);
 
diff --git a/Engine/Source/Runtime/Engine/Classes/Engine/SkeletalMesh.h b/Engine/Source/Runtime/Engine/Classes/Engine/SkeletalMesh.h
--- a/Engine/Source/Runtime/Engine/Classes/Engine/SkeletalMesh.h
+++ b/Engine/Source/Runtime/Engine/Classes/Engine/SkeletalMesh.h
@@ -922,6 +922,10 @@ public:
 #if WITH_EDITOR
 	/** Calculate the required bones for a Skeletal Mesh LOD, including possible extra influences */
 	static void CalculateRequiredBones(FSkeletalMeshLODModel& LODModel, const struct FReferenceSkeleton& RefSkeleton, const TMap<FBoneIndexType, FBoneIndexType> * BonesToRemove);
+
+	/** Recalculate Retarget Base Pose BoneTransform */
+	void ReallocateRetargetBasePose();
+
 #endif // WITH_EDITOR
 
 	/** 
diff --git a/Engine/Source/Runtime/Engine/Private/Animation/AnimNode_CustomProperty.cpp b/Engine/Source/Runtime/Engine/Private/Animation/AnimNode_CustomProperty.cpp
--- a/Engine/Source/Runtime/Engine/Private/Animation/AnimNode_CustomProperty.cpp
+++ b/Engine/Source/Runtime/Engine/Private/Animation/AnimNode_CustomProperty.cpp
@@ -7,6 +7,9 @@
 
 FAnimNode_CustomProperty::FAnimNode_CustomProperty()
 	: TargetInstance(nullptr)
+#if WITH_EDITOR
+	, bReinitializeProperties(false)
+#endif // WITH_EDITOR
 {
 
 }
@@ -42,6 +45,19 @@ void FAnimNode_CustomProperty::PropagateInputProperties(const UObject* InSourceI
 	}
 }
 
+void FAnimNode_CustomProperty::PreUpdate(const UAnimInstance* InAnimInstance) 
+{
+	FAnimNode_Base::PreUpdate(InAnimInstance);
+
+#if WITH_EDITOR
+	if (bReinitializeProperties)
+	{
+		InitializeProperties(InAnimInstance);
+		bReinitializeProperties = false;
+	}
+#endif// WITH_EDITOR
+}
+
 void FAnimNode_CustomProperty::InitializeProperties(const UObject* InSourceInstance)
 {
 	UClass* TargetClass = GetTargetClass();
diff --git a/Engine/Source/Runtime/Engine/Private/Animation/Skeleton.cpp b/Engine/Source/Runtime/Engine/Private/Animation/Skeleton.cpp
--- a/Engine/Source/Runtime/Engine/Private/Animation/Skeleton.cpp
+++ b/Engine/Source/Runtime/Engine/Private/Animation/Skeleton.cpp
@@ -1660,8 +1660,13 @@ void USkeleton::HandleVirtualBoneChanges()
 		USkeletalMesh* SkelMesh = *ItMesh;
 		if (SkelMesh->Skeleton == this)
 		{
+			// also have to update retarget base pose
 			SkelMesh->RefSkeleton.RebuildRefSkeleton(this, bRebuildNameMap);
 			RebuildLinkup(SkelMesh);
+#if WITH_EDITOR
+			// whole bone count has changed, so it has to recalculate retarget base pose
+			SkelMesh->ReallocateRetargetBasePose();
+#endif // #if WITH_EDITOR
 		}
 	}
 
diff --git a/Engine/Source/Runtime/PropertyPath/Private/PropertyPathHelpers.cpp b/Engine/Source/Runtime/PropertyPath/Private/PropertyPathHelpers.cpp
--- a/Engine/Source/Runtime/PropertyPath/Private/PropertyPathHelpers.cpp
+++ b/Engine/Source/Runtime/PropertyPath/Private/PropertyPathHelpers.cpp
@@ -814,6 +814,11 @@ void FCachedPropertyPath::RemoveFromStart(int32 InNumSegments)
 	}
 }
 
+UProperty* FCachedPropertyPath::GetUProperty() const
+{
+	return Cast<UProperty>(GetLastSegment().GetField());
+}
+
 namespace PropertyPathHelpers
 {
 	void FindFieldNameAndArrayIndex(int32 InCount, const TCHAR* InString, int32& OutCount, const TCHAR** OutPropertyName, int32& OutArrayIndex)
diff --git a/Engine/Source/Runtime/PropertyPath/Public/PropertyPathHelpers.h b/Engine/Source/Runtime/PropertyPath/Public/PropertyPathHelpers.h
--- a/Engine/Source/Runtime/PropertyPath/Public/PropertyPathHelpers.h
+++ b/Engine/Source/Runtime/PropertyPath/Public/PropertyPathHelpers.h
@@ -175,6 +175,8 @@ struct PROPERTYPATH_API FCachedPropertyPath
 	/** Trims this property path at the start */
 	void RemoveFromStart(int32 InNumSegments = 1);
 
+	/** Returns UProperty if valid. This can be UFunction */
+	UProperty* GetUProperty() const;
 private:
 	/** Path segments for this path */
 	UPROPERTY()
