Fixed a bug as uncovered by the new unit test (to be checked in in a separate submit).

diff --git a/Engine/Source/Runtime/AnimGraphRuntime/Private/KismetAnimationLibrary.cpp b/Engine/Source/Runtime/AnimGraphRuntime/Private/KismetAnimationLibrary.cpp
--- a/Engine/Source/Runtime/AnimGraphRuntime/Private/KismetAnimationLibrary.cpp
+++ b/Engine/Source/Runtime/AnimGraphRuntime/Private/KismetAnimationLibrary.cpp
@@ -153,7 +153,7 @@ float UKismetAnimationLibrary::K2_CalculateVelocityFromPositionHistory(
 				NextIndex = 0;
 			}
 			History.Positions[NextIndex] = Position;
-			History.Velocities[History.LastIndex] = LengthOfV;
+			History.Velocities[NextIndex] = LengthOfV;
 			History.LastIndex = NextIndex;
 		}
 		else
