Updated two headers to comply with the DocumentationPolicy = "Strict" setting as seen in https://p4-swarm.epicgames.net/reviews/5097025

All of the changes are enforced by the new setting for the UnrealHeaderTool.

diff --git a/Engine/Source/Runtime/AnimGraphRuntime/Public/KismetAnimationLibrary.h b/Engine/Source/Runtime/AnimGraphRuntime/Public/KismetAnimationLibrary.h
--- a/Engine/Source/Runtime/AnimGraphRuntime/Public/KismetAnimationLibrary.h
+++ b/Engine/Source/Runtime/AnimGraphRuntime/Public/KismetAnimationLibrary.h
@@ -10,69 +10,200 @@
 
 class USkeletalMeshComponent;
 
-UCLASS(meta=(ScriptName="AnimGraphLibrary"))
+/**
+	A library of the most common animation blueprint functions.
+*/
+UCLASS(meta=(ScriptName="AnimGraphLibrary", DocumentationPolicy="Strict"))
 class ANIMGRAPHRUNTIME_API UKismetAnimationLibrary : public UBlueprintFunctionLibrary
 {
 	GENERATED_UCLASS_BODY()
 
+	/**
+		Computes the transform for two bones using inverse kinematics.
+
+		@param RootPos The input root position of the two bone chain
+		@param JointPos The input center (elbow) position of the two bone chain
+		@param EndPos The input end (wrist) position of the two bone chain
+		@param JointTarget The IK target for the write to reach
+		@param Effector The position of the target effector for the IK Chain.
+		@param OutJointPos The resulting position for the center (elbow)
+		@param OutEndPos The resulting position for the end (wrist)
+		@param bAllowStretching If set to true the bones are allowed to stretch
+		@param StartStretchRatio The ratio at which the bones should start to stretch. The higher the value, the later the stretching wil start.
+		@param MaxStretchScale The maximum multiplier for the stretch to reach.
+	 */
 	UFUNCTION(BlueprintPure, Category = "Utilities|Animation", meta = (DisplayName = "Two Bone IK Function", ScriptName = "TwoBoneIK", bAllowStretching = "false", StartStretchRatio = "1.0", MaxStretchScale = "1.2"))
-	static void K2_TwoBoneIK(const FVector& RootPos, const FVector& JointPos, const FVector& EndPos, const FVector& JointTarget, const FVector& Effector, FVector& OutJointPos, FVector& OutEndPos, bool bAllowStretching, float StartStretchRatio, float MaxStretchScale);
+	static void K2_TwoBoneIK(
+		const FVector& RootPos, 
+		const FVector& JointPos, 
+		const FVector& EndPos, 
+		const FVector& JointTarget, 
+		const FVector& Effector, 
+		FVector& OutJointPos, 
+		FVector& OutEndPos, 
+		bool bAllowStretching, 
+		UPARAM(meta = (UIMin = "0.75", UIMax = "1.0")) float StartStretchRatio,
+		UPARAM(meta = (UIMin = "1.0", UIMax = "5.0")) float MaxStretchScale
+	);
+
+	/**
+		Computes the transform which is "looking" at target position with a local axis.
 
+		@param CurrentTransform The input transform to modify
+		@param TargetPosition The position this transform should look at
+		@param LookAtVector The local vector to align with the target
+		@param bUseUpVector If set to true the lookat will also perform a twist rotation
+		@param UpVector The position to use for the upvector target (only used is bUseUpVector is turned on)
+		@param ClampConeInDegree A limit for only allowing the lookat to rotate as much as defined by the float value
+		@return The transform (based on the CurrentTransform) with the modified rotation
+	*/
 	UFUNCTION(BlueprintPure, Category = "Utilities|Animation", meta = (DisplayName = "Look At Function", ScriptName = "LookAt", bUseUpVector = "false"))
-	static FTransform K2_LookAt(const FTransform& CurrentTransform, const FVector& TargetPosition, FVector LookAtVector, bool bUseUpVector, FVector UpVector, float ClampConeInDegree);
+	static FTransform K2_LookAt(
+		const FTransform& CurrentTransform, 
+		const FVector& TargetPosition, 
+		FVector LookAtVector, 
+		bool bUseUpVector, 
+		FVector UpVector, 
+		UPARAM(meta = (UIMin = "0.0", UIMax = "180.0")) float ClampConeInDegree
+	);
 
+	/**
+		Computes the distance between two bones / sockets and can remap the range.
+
+		@param Component The skeletal component to look for the sockets / bones within
+		@param SocketOrBoneNameA The name of the first socket / bone
+		@param SocketSpaceA The space for the first socket / bone
+		@param SocketOrBoneNameB The name of the second socket / bone
+		@param SocketSpaceB The space for the second socket / bone
+		@param bRemapRange If set to true, the distance will be remapped using the range parameters
+		@param InRangeMin The minimum for the input range (commonly == 0.0)
+		@param InRangeMax The maximum for the input range (the max expected distance)
+		@param OutRangeMin The minimum for the output range (commonly == 0.0)
+		@param OutRangeMax The maximum for the output range (commonly == 1.0)
+		@return The distance or the remapped value.
+	*/
 	UFUNCTION(BlueprintPure, Category = "Utilities|Animation", meta = (DisplayName = "Get Distance Between Two Sockets", ScriptName = "DistanceBetweenSockets", bRemapRange= "false"))
-	static float K2_DistanceBetweenTwoSocketsAndMapRange(const USkeletalMeshComponent* Component, const FName SocketOrBoneNameA, ERelativeTransformSpace SocketSpaceA, const FName SocketOrBoneNameB, ERelativeTransformSpace SocketSpaceB, bool bRemapRange, float InRangeMin, float InRangeMax, float OutRangeMin, float OutRangeMax);
-
+	static float K2_DistanceBetweenTwoSocketsAndMapRange(
+		const USkeletalMeshComponent* Component, 
+		const FName SocketOrBoneNameA, 
+		ERelativeTransformSpace SocketSpaceA, 
+		const FName SocketOrBoneNameB, 
+		ERelativeTransformSpace SocketSpaceB, 
+		bool bRemapRange, 
+		UPARAM(meta = (UIMin = "0.0", UIMax = "1000.0")) float InRangeMin,
+		UPARAM(meta = (UIMin = "0.0", UIMax = "1000.0")) float InRangeMax,
+		UPARAM(meta = (UIMin = "0.0", UIMax = "4.0")) float OutRangeMin,
+		UPARAM(meta = (UIMin = "0.0", UIMax = "4.0")) float OutRangeMax);
+
+	/**
+		Computes the direction between two bones / sockets.
+
+		@param Component The skeletal component to look for the sockets / bones within
+		@param SocketOrBoneNameFrom The name of the first socket / bone
+		@param SocketOrBoneNameTo The name of the second socket / bone
+		@return The direction between the two sockets / bones ( SocketOrBoneNameTo.Location - SocketOrBoneNameFrom.Location )
+	*/
 	UFUNCTION(BlueprintPure, Category = "Utilities|Animation", meta = (DisplayName = "Get Direction Between Sockets", ScriptName = "DirectionBetweenSockets"))
 	static FVector K2_DirectionBetweenSockets(const USkeletalMeshComponent* Component, const FName SocketOrBoneNameFrom, const FName SocketOrBoneNameTo);
 
-	/* This function creates perlin noise from input X, Y, Z, and then range map to RangeOut, and out put to OutX, OutY, OutZ */
+	/** 
+		This function creates perlin noise from input X, Y, Z, and then range map to RangeOut, and out put to OutX, OutY, OutZ
+
+		@param X The x component for the input position of the noise
+		@param Y The y component for the input position of the noise
+		@param Z The z component for the input position of the noise
+		@param RangeOutMinX The minimum for the output range for the x component
+		@param RangeOutMaxX The maximum for the output range for the x component
+		@param RangeOutMinY The minimum for the output range for the y component
+		@param RangeOutMaxY The maximum for the output range for the y component
+		@param RangeOutMinZ The minimum for the output range for the z component
+		@param RangeOutMaxZ The maximum for the output range for the z component
+		@return The noise value based on the input position and ranges
+	*/
 	UFUNCTION(BlueprintPure, Category = "Utilities|Animation", meta = (DisplayName = "Make Perlin Noise Vector and Remap", ScriptName = "MakeVectorFromPerlinNoise", RangeOutMinX = "-1.f", RangeOutMaxX = "1.f", RangeOutMinY = "-1.f", RangeOutMaxY = "1.f", RangeOutMinZ = "-1.f", RangeOutMaxZ = "1.f"))
-	static FVector K2_MakePerlinNoiseVectorAndRemap(float X, float Y, float Z, float RangeOutMinX, float RangeOutMaxX, float RangeOutMinY, float RangeOutMaxY, float RangeOutMinZ, float RangeOutMaxZ);
+	static FVector K2_MakePerlinNoiseVectorAndRemap(
+		UPARAM(meta = (UIMin = "0.0", UIMax = "0.0")) float X,
+		UPARAM(meta = (UIMin = "0.0", UIMax = "0.0")) float Y,
+		UPARAM(meta = (UIMin = "0.0", UIMax = "0.0")) float Z,
+		UPARAM(meta = (UIMin = "-100.0", UIMax = "100.0")) float RangeOutMinX,
+		UPARAM(meta = (UIMin = "-100.0", UIMax = "100.0")) float RangeOutMaxX,
+		UPARAM(meta = (UIMin = "-100.0", UIMax = "100.0")) float RangeOutMinY,
+		UPARAM(meta = (UIMin = "-100.0", UIMax = "100.0")) float RangeOutMaxY,
+		UPARAM(meta = (UIMin = "-100.0", UIMax = "100.0")) float RangeOutMinZ,
+		UPARAM(meta = (UIMin = "-100.0", UIMax = "100.0")) float RangeOutMaxZ
+	);
+
+	/**
+		This function creates perlin noise for a single float and then range map to RangeOut
 
+		@param Value The input value for the noise function
+		@param RangeOutMin The minimum for the output range
+		@param RangeOutMax The maximum for the output range
+		@return The noise value based on the input value and range
+	*/
 	UFUNCTION(BlueprintPure, Category = "Utilities|Animation", meta = (DisplayName = "Make Perlin Noise and Remap", ScriptName = "MakeFloatFromPerlinNoise", RangeOutMinX = "-1.f", RangeOutMaxX = "1.f", RangeOutMinY = "-1.f", RangeOutMaxY = "1.f", RangeOutMinZ = "-1.f", RangeOutMaxZ = "1.f"))
-	static float K2_MakePerlinNoiseAndRemap(float Value, float RangeOutMin, float RangeOutMax);
+	static float K2_MakePerlinNoiseAndRemap(
+		UPARAM(meta = (UIMin = "0.0", UIMax = "0.0")) float Value,
+		UPARAM(meta = (UIMin = "-100.0", UIMax = "100.0")) float RangeOutMin,
+		UPARAM(meta = (UIMin = "-100.0", UIMax = "100.0")) float RangeOutMax
+	);
 
 	/** 
 		This function calculates the velocity of a position changing over time.
 		You need to hook up a valid PositionHistory variable to this for storage.
 
+		@param DeltaSeconds The time passed in seconds
 		@param Position The position to track over time.
 		@param History The history to use for storage.
 		@param NumberOfSamples The number of samples to use for the history. The higher the number of samples - the smoother the velocity changes.
 		@param VelocityMin The minimum velocity to use for normalization (if both min and max are set to 0, normalization is turned off)
 		@param VelocityMax The maximum velocity to use for normalization (if both min and max are set to 0, normalization is turned off)
+		@return The magnitude of the calculated velocity
 	*/
 	UFUNCTION(BlueprintCallable, Category = "Utilities|Animation", meta = (DisplayName = "Calculate Velocity From Position History", ScriptName = "CalculateVelocityFromPositionHistory", NumberOfSamples = "16", VelocityMin = "0.f", VelocityMax = "128.f"))
 	static float K2_CalculateVelocityFromPositionHistory(
-		float DeltaSeconds,
+		UPARAM(meta = (UIMin = "0.0", UIMax = "0.0")) float DeltaSeconds,
 		FVector Position,
 		UPARAM(ref) FPositionHistory& History,
 		int32 NumberOfSamples,
-		float VelocityMin,
-		float VelocityMax
+		UPARAM(meta = (UIMin = "0.0", UIMax = "256.0")) float VelocityMin,
+		UPARAM(meta = (UIMin = "0.0", UIMax = "256.0")) float VelocityMax
 	);
 
+	/**
+		This function perform easing on a float value using a variety of easing types.
+
+		@param Value The float value to ease
+		@param EasingType The easing function to use
+		@return The modified float value
+	*/
 	UFUNCTION(BlueprintPure, Category = "Utilities|Animation", meta = (DisplayName = "Scalar Easing", ScriptName = "ScalarEasing", Value = "0.f"))
-	static float K2_ScalarEasing(float Value, EEasingFuncType EasingType);
+	static float K2_ScalarEasing(
+		UPARAM(meta = (UIMin = "0.0", UIMax = "0.0")) float Value, 
+		EEasingFuncType EasingType
+	);
 
 	/** 
 		This function calculates the velocity of an offset position on a bone / socket over time.
 		The bone's / socket's motion can be expressed within a reference frame (another bone / socket). 
 		You need to hook up a valid PositionHistory variable to this for storage.
 
+		@param DeltaSeconds The time passed in seconds
+		@param Component The skeletal component to look for the bones / sockets
 		@param SocketOrBoneName The name of the bone / socket to track.
 		@param ReferenceSocketOrBone The name of the bone / socket to use as a frame of reference (or None if no frame of reference == world space).
+		@param SocketSpace The space to use for the two sockets / bones
 		@param OffsetInBoneSpace The relative position in the space of the bone / socket to track over time.
 		@param History The history to use for storage.
 		@param NumberOfSamples The number of samples to use for the history. The higher the number of samples - the smoother the velocity changes.
 		@param VelocityMin The minimum velocity to use for normalization (if both min and max are set to 0, normalization is turned off)
 		@param VelocityMax The maximum velocity to use for normalization (if both min and max are set to 0, normalization is turned off)
+		@param EasingType The easing function to use
+		@return The magnitude of the calculated velocity
 	*/
 	UFUNCTION(BlueprintCallable, Category = "Utilities|Animation", meta = (DisplayName = "Calculate Velocity From Sockets", ScriptName = "CalculateVelocityFromSockets", NumberOfSamples = "16", VelocityMin = "0.f", VelocityMax = "128.f"))
 	static float K2_CalculateVelocityFromSockets(
-		float DeltaSeconds, 
+		UPARAM(meta = (UIMin = "0.0", UIMax = "0.0")) float DeltaSeconds,
 		USkeletalMeshComponent * Component,
 		const FName SocketOrBoneName,
 		const FName ReferenceSocketOrBone,
@@ -80,8 +211,8 @@ class ANIMGRAPHRUNTIME_API UKismetAnimationLibrary : public UBlueprintFunctionLi
 		FVector OffsetInBoneSpace, 
 		UPARAM(ref) FPositionHistory& History, 
 		int32 NumberOfSamples,
-		float VelocityMin, 
-		float VelocityMax,
+		UPARAM(meta = (UIMin = "0.0", UIMax = "256.0")) float VelocityMin,
+		UPARAM(meta = (UIMin = "0.0", UIMax = "256.0")) float VelocityMax,
 		EEasingFuncType EasingType
 	);
 };
diff --git a/Engine/Source/Runtime/AnimGraphRuntime/Public/KismetAnimationTypes.h b/Engine/Source/Runtime/AnimGraphRuntime/Public/KismetAnimationTypes.h
--- a/Engine/Source/Runtime/AnimGraphRuntime/Public/KismetAnimationTypes.h
+++ b/Engine/Source/Runtime/AnimGraphRuntime/Public/KismetAnimationTypes.h
@@ -5,27 +5,49 @@
 #include "CoreMinimal.h"
 #include "KismetAnimationTypes.generated.h"
 
-// The kind of easing function to use
-UENUM()
+/**
+	An easing type defining how to ease float values.
+*/
+UENUM(meta = (DocumentationPolicy = "Strict"))
 enum class EEasingFuncType : uint8
 {
+	// Linear easing (no change to the value)
 	Linear,
+	// Easing using a sinus function
 	Sinusoidal,
+	// Cubic version of the value (only in)
 	Cubic,
+	// Quadratic version of the value (in and out)
 	QuadraticInOut,
+	// Cubic version of the value (in and out)
 	CubicInOut,
+	// Easing using a cubic hermite function
 	HermiteCubic,
+	// Quartic version of the value (in and out)
 	QuarticInOut,
+	// Quintic version of the value (in and out)
 	QuinticInOut,
+	// Circular easing (only in)
 	CircularIn,
+	// Circular easing (only out)
 	CircularOut,
+	// Circular easing (in and out)
 	CircularInOut,
+	// Exponential easing (only in)
 	ExpIn,
+	// Exponential easing (only out)
 	ExpOut,
+	// Exponential easing (in and out)
 	ExpInOut
 };
 
-USTRUCT(BlueprintType)
+/**
+	The FPositionHistory is a container to record position changes
+	over time. This is used to calculate velocity of a bone, for example.
+	The FPositionArray also tracks the last index used to allow for 
+	reuse of entries (instead of appending to the array all the time).
+*/
+USTRUCT(BlueprintType, meta = (DocumentationPolicy = "Strict"))
 struct FPositionHistory
 {
 	GENERATED_BODY()
@@ -43,6 +65,10 @@ public:
 	UPROPERTY(EditAnywhere, Category = "FPositionHistory", meta = (MetaClass = "AdvancedCopyCustomization"))
 	TArray<FVector> Positions;
 
+	/** The range for this particular history */
+	UPROPERTY(EditAnywhere, Category = "FPositionHistory", meta = (MetaClass = "AdvancedCopyCustomization", UIMin = "0.0", UIMax = "1.0"))
+	float Range;
+
 	TArray<float> Velocities;
 	uint32 LastIndex; /// The last index used to store a position
 };
