Adding initialization to range member.

diff --git a/Engine/Source/Runtime/AnimGraphRuntime/Public/KismetAnimationTypes.h b/Engine/Source/Runtime/AnimGraphRuntime/Public/KismetAnimationTypes.h
--- a/Engine/Source/Runtime/AnimGraphRuntime/Public/KismetAnimationTypes.h
+++ b/Engine/Source/Runtime/AnimGraphRuntime/Public/KismetAnimationTypes.h
@@ -57,6 +57,7 @@ public:
 	/** Default constructor */
 	FPositionHistory()
 		: Positions(TArray<FVector>())
+		, Range(0.f)
 		, Velocities(TArray<float>())
 		, LastIndex(0)
 	{}
