Build break fix for header change

diff --git a/Engine/Source/Runtime/AnimGraphRuntime/Public/KismetAnimationTypes.h b/Engine/Source/Runtime/AnimGraphRuntime/Public/KismetAnimationTypes.h
--- a/Engine/Source/Runtime/AnimGraphRuntime/Public/KismetAnimationTypes.h
+++ b/Engine/Source/Runtime/AnimGraphRuntime/Public/KismetAnimationTypes.h
@@ -8,7 +8,6 @@
 
 /**
  *	An easing type defining how to ease float values.
- */
  *	The FPositionHistory is a container to record position changes
  *	over time. This is used to calculate velocity of a bone, for example.
  *	The FPositionArray also tracks the last index used to allow for 
