Grey out options for springs based on checkbox

diff --git a/Engine/Source/Runtime/AnimGraphRuntime/Public/BoneControllers/AnimNode_AnimDynamics.h b/Engine/Source/Runtime/AnimGraphRuntime/Public/BoneControllers/AnimNode_AnimDynamics.h
--- a/Engine/Source/Runtime/AnimGraphRuntime/Public/BoneControllers/AnimNode_AnimDynamics.h
+++ b/Engine/Source/Runtime/AnimGraphRuntime/Public/BoneControllers/AnimNode_AnimDynamics.h
@@ -285,15 +285,19 @@ struct ANIMGRAPHRUNTIME_API FAnimNode_AnimDynamics : public FAnimNode_SkeletalCo
 	UPROPERTY(EditAnywhere, BlueprintReadWrite, Category = Setup, meta = (PinHiddenByDefault, EditCondition = "bUseGravityOverride"))
 	FVector GravityOverride;
 
-	/** Spring constant to use when calculating linear springs, higher values mean a stronger spring.*/
-	UPROPERTY(EditAnywhere, BlueprintReadWrite, Category = Setup, meta = (PinHiddenByDefault, DisplayAfter="bAngularSpring"))
+	/** 
+	 * Spring constant to use when calculating linear springs, higher values mean a stronger spring.
+	 * You need to enable the Linear Spring checkbox for this to have an effect.
+	 */
+	UPROPERTY(EditAnywhere, BlueprintReadWrite, Category = Setup, meta = (PinHiddenByDefault, DisplayAfter="bAngularSpring", EditCondition = "bLinearSpring"))
 	float LinearSpringConstant;
 
 	/** 
 	 * Spring constant to use when calculating angular springs, higher values mean a stronger spring.
+	 * You need to enable the Angular Spring checkbox for this to have an effect.
 	 * Note: Make sure to also set the Angular Target Axis and Angular Target in the Constraint Setup for this to have an effect.
 	 */
-	UPROPERTY(EditAnywhere, BlueprintReadWrite, Category = Setup, meta = (PinHiddenByDefault))
+	UPROPERTY(EditAnywhere, BlueprintReadWrite, Category = Setup, meta = (PinHiddenByDefault, EditCondition = "bAngularSpring"))
 	float AngularSpringConstant;
 
 	/** Scale to apply to calculated wind velocities in the solver */
