Fix additive blending not working with less than 1 weight

diff --git a/Engine/Source/Runtime/AnimGraphRuntime/Private/AnimNodes/AnimNode_MultiWayBlend.cpp b/Engine/Source/Runtime/AnimGraphRuntime/Private/AnimNodes/AnimNode_MultiWayBlend.cpp
--- a/Engine/Source/Runtime/AnimGraphRuntime/Private/AnimNodes/AnimNode_MultiWayBlend.cpp
+++ b/Engine/Source/Runtime/AnimGraphRuntime/Private/AnimNodes/AnimNode_MultiWayBlend.cpp
@@ -136,6 +136,8 @@ void FAnimNode_MultiWayBlend::Evaluate_AnyThread(FPoseContext& Output)
 	if (SourcePoses.Num() > 0)
 	{
 		FAnimationRuntime::BlendPosesTogether(SourcePoses, SourceCurves, SourceWeights, Output.Pose, Output.Curve);
+		// normalize rotation - some cases, where additive is applied less than 1, it will use non normalized rotation
+		Output.Pose.NormalizeRotations();
 	}
 	else
 	{
diff --git a/Engine/Source/Runtime/AnimGraphRuntime/Private/AnimSequencerInstanceProxy.cpp b/Engine/Source/Runtime/AnimGraphRuntime/Private/AnimSequencerInstanceProxy.cpp
--- a/Engine/Source/Runtime/AnimGraphRuntime/Private/AnimSequencerInstanceProxy.cpp
+++ b/Engine/Source/Runtime/AnimGraphRuntime/Private/AnimSequencerInstanceProxy.cpp
@@ -144,6 +144,13 @@ void FAnimSequencerInstanceProxy::UpdateAnimTrack(UAnimSequenceBase* InAnimSeque
 	// if moving to 0.f, we mark this to teleport. Otherwise, do not use explicit time
 	FAnimNode_MultiWayBlend& BlendNode = (PlayerState->bAdditive) ? AdditiveBlendNode : FullBodyBlendNode;
 	BlendNode.DesiredAlphas[PlayerState->PoseIndex] = Weight;
+
+	// if additive, apply alpha value correctly
+	// this will be used when apply additive is blending correct total alpha to additive
+	if (PlayerState->bAdditive)
+	{
+		SequencerRootNode.Alpha = BlendNode.GetTotalAlpha();
+	}
 }
 
 void FAnimSequencerInstanceProxy::EnsureAnimTrack(UAnimSequenceBase* InAnimSequence, uint32 SequenceId)
diff --git a/Engine/Source/Runtime/AnimGraphRuntime/Public/AnimNodes/AnimNode_MultiWayBlend.h b/Engine/Source/Runtime/AnimGraphRuntime/Public/AnimNodes/AnimNode_MultiWayBlend.h
--- a/Engine/Source/Runtime/AnimGraphRuntime/Public/AnimNodes/AnimNode_MultiWayBlend.h
+++ b/Engine/Source/Runtime/AnimGraphRuntime/Public/AnimNodes/AnimNode_MultiWayBlend.h
@@ -69,8 +69,8 @@ public:
 		Poses.Reset();
 		DesiredAlphas.Reset();
 	}
-private:
-	float GetTotalAlpha()
+
+	float GetTotalAlpha() const
 	{
 		float TotalAlpha = 0.f;
 
@@ -82,6 +82,7 @@ private:
 		return TotalAlpha;
 	}
 
+private:
 	// process new weights and then return out
 	void UpdateCachedAlphas();
 };
