Fix to get back to the previous pose that a SkeletalMesh has before getting evaluated.

Added a PoseNode to the AnimInstanceProxy that we save a pose to right when we initialize the seuqencer proxy and then in RestoreState in the template we then link that anim pose to use that stored pose and then  force an evaluation.

diff --git a/Engine/Source/Runtime/AnimGraphRuntime/Private/AnimSequencerInstance.cpp b/Engine/Source/Runtime/AnimGraphRuntime/Private/AnimSequencerInstance.cpp
--- a/Engine/Source/Runtime/AnimGraphRuntime/Private/AnimSequencerInstance.cpp
+++ b/Engine/Source/Runtime/AnimGraphRuntime/Private/AnimSequencerInstance.cpp
@@ -12,6 +12,8 @@
 // UAnimSequencerInstance
 /////////////////////////////////////////////////////
 
+const FName UAnimSequencerInstance::SequencerPoseName(TEXT("Sequencer_Pose_Name"));
+
 UAnimSequencerInstance::UAnimSequencerInstance(const FObjectInitializer& ObjectInitializer)
 	: Super(ObjectInitializer)
 {
@@ -31,4 +33,19 @@ void UAnimSequencerInstance::UpdateAnimTrack(UAnimSequenceBase* InAnimSequence,
 void UAnimSequencerInstance::ResetNodes()
 {
 	GetProxyOnGameThread<FAnimSequencerInstanceProxy>().ResetNodes();
-}
\ No newline at end of file
+}
+
+void UAnimSequencerInstance::ResetPose()
+{
+	GetProxyOnGameThread<FAnimSequencerInstanceProxy>().ResetPose();
+}
+
+void UAnimSequencerInstance::NativeInitializeAnimation()
+{
+	SavePose();
+}
+
+void UAnimSequencerInstance::SavePose()
+{
+	SavePoseSnapshot(UAnimSequencerInstance::SequencerPoseName);
+}
diff --git a/Engine/Source/Runtime/AnimGraphRuntime/Private/AnimSequencerInstanceProxy.cpp b/Engine/Source/Runtime/AnimGraphRuntime/Private/AnimSequencerInstanceProxy.cpp
--- a/Engine/Source/Runtime/AnimGraphRuntime/Private/AnimSequencerInstanceProxy.cpp
+++ b/Engine/Source/Runtime/AnimGraphRuntime/Private/AnimSequencerInstanceProxy.cpp
@@ -40,6 +40,10 @@ void FAnimSequencerInstanceProxy::ConstructNodes()
 	FullBodyBlendNode.ResetPoses();
 	AdditiveBlendNode.ResetPoses();
 
+
+	SnapshotNode.SnapshotName = UAnimSequencerInstance::SequencerPoseName;
+
+
 	ClearSequencePlayerMap();
 }
 
@@ -53,6 +57,11 @@ void FAnimSequencerInstanceProxy::ClearSequencePlayerMap()
 	SequencerToPlayerMap.Empty();
 }
 
+void FAnimSequencerInstanceProxy::ResetPose()
+{
+	SequencerRootNode.Base.SetLinkNode(&SnapshotNode);
+	//force evaluation?
+}	
 void FAnimSequencerInstanceProxy::ResetNodes()
 {
 	FMemory::Memzero(FullBodyBlendNode.DesiredAlphas.GetData(), FullBodyBlendNode.DesiredAlphas.GetAllocatedSize());
@@ -157,4 +166,5 @@ void FAnimSequencerInstanceProxy::EnsureAnimTrack(UAnimSequenceBase* InAnimSeque
 	{
 		PlayerState->PlayerNode.OverrideAsset(InAnimSequence);
 	}
-}
\ No newline at end of file
+}
+
diff --git a/Engine/Source/Runtime/AnimGraphRuntime/Public/AnimSequencerInstance.h b/Engine/Source/Runtime/AnimGraphRuntime/Public/AnimSequencerInstance.h
--- a/Engine/Source/Runtime/AnimGraphRuntime/Public/AnimSequencerInstance.h
+++ b/Engine/Source/Runtime/AnimGraphRuntime/Public/AnimSequencerInstance.h
@@ -15,14 +15,27 @@ UCLASS(transient, NotBlueprintable)
 class ANIMGRAPHRUNTIME_API UAnimSequencerInstance : public UAnimCustomInstance
 {
 	GENERATED_UCLASS_BODY()
+
 public:
+
 	/** Update an animation sequence player in this instance */
 	void UpdateAnimTrack(UAnimSequenceBase* InAnimSequence, int32 SequenceId, float InPosition, float Weight, bool bFireNotifies);
 
 	/** Reset all nodes in this instance */
 	void ResetNodes();
 
+	/** Reset the pose for this instance*/
+	void ResetPose();
+
 protected:
 	// UAnimInstance interface
 	virtual FAnimInstanceProxy* CreateAnimInstanceProxy() override;
+	virtual void NativeInitializeAnimation() override;
+
+	/** Saved the named pose to restore after */
+	void SavePose();
+
+public:
+	static const FName SequencerPoseName;
+
 };
diff --git a/Engine/Source/Runtime/AnimGraphRuntime/Public/AnimSequencerInstanceProxy.h b/Engine/Source/Runtime/AnimGraphRuntime/Public/AnimSequencerInstanceProxy.h
--- a/Engine/Source/Runtime/AnimGraphRuntime/Public/AnimSequencerInstanceProxy.h
+++ b/Engine/Source/Runtime/AnimGraphRuntime/Public/AnimSequencerInstanceProxy.h
@@ -6,6 +6,7 @@
 #include "AnimNodes/AnimNode_SequenceEvaluator.h"
 #include "AnimNodes/AnimNode_ApplyAdditive.h"
 #include "AnimNodes/AnimNode_MultiWayBlend.h"
+#include "AnimNodes/AnimNode_PoseSnapshot.h"
 #include "AnimSequencerInstanceProxy.generated.h"
 
 /** Base class for all 'players' that can attach to and be blended into a sequencer instance's output */
@@ -86,6 +87,9 @@ public:
 	/** Reset all nodes in this instance */
 	virtual void ResetNodes();
 
+	/** Reset the pose in this instance*/
+	virtual void ResetPose();
+
 protected:
 	/** Find a player of a specified type */
 	template<typename Type>
@@ -107,6 +111,7 @@ protected:
 	struct FAnimNode_ApplyAdditive SequencerRootNode;
 	struct FAnimNode_MultiWayBlend FullBodyBlendNode;
 	struct FAnimNode_MultiWayBlend AdditiveBlendNode;
+	struct FAnimNode_PoseSnapshot  SnapshotNode;
 
 	/** mapping from sequencer index to internal player index */
 	TMap<uint32, FSequencerPlayerBase*> SequencerToPlayerMap;
diff --git a/Engine/Source/Runtime/MovieSceneTracks/Private/Evaluation/MovieSceneSkeletalAnimationTemplate.cpp b/Engine/Source/Runtime/MovieSceneTracks/Private/Evaluation/MovieSceneSkeletalAnimationTemplate.cpp
--- a/Engine/Source/Runtime/MovieSceneTracks/Private/Evaluation/MovieSceneSkeletalAnimationTemplate.cpp
+++ b/Engine/Source/Runtime/MovieSceneTracks/Private/Evaluation/MovieSceneSkeletalAnimationTemplate.cpp
@@ -91,6 +91,15 @@ struct FPreAnimatedAnimationTokenProducer : IMovieScenePreAnimatedTokenProducer
 				UAnimSequencerInstance* SequencerInst = Cast<UAnimSequencerInstance>(Component->GetAnimInstance());
 				if (SequencerInst)
 				{
+					SequencerInst->ResetPose();
+					//wish this 'force evaluation' was encapsulated better
+					Component->TickAnimation(0.f, false);
+					Component->RefreshBoneTransforms();
+					Component->RefreshSlaveComponents();
+					Component->UpdateComponentToWorld();
+					Component->FinalizeBoneTransform();
+					Component->MarkRenderTransformDirty();
+					Component->MarkRenderDynamicDataDirty();
 					SequencerInst->ResetNodes();
 				}
 
