Fixed pose exploding when chaining multiple Blend Multi nodes together. FAnimNode_MultiWayBlend uses scratch arrays to build up input pose data before blending all the poses together. But when multiple Blend Multi nodes were chained together they would stomp on each other by resetting the arrays on entry. Now the scratch arrays are treated like stacks and we push each source pose's data and pop what was added before leaving the evaluate call.

[RN] fix Animation Fixed bad pose output when chaining multiple Blend Multi nodes together.

diff --git a/Engine/Source/Runtime/AnimGraphRuntime/Private/AnimNodes/AnimNode_MultiWayBlend.cpp b/Engine/Source/Runtime/AnimGraphRuntime/Private/AnimNodes/AnimNode_MultiWayBlend.cpp
--- a/Engine/Source/Runtime/AnimGraphRuntime/Private/AnimNodes/AnimNode_MultiWayBlend.cpp
+++ b/Engine/Source/Runtime/AnimGraphRuntime/Private/AnimNodes/AnimNode_MultiWayBlend.cpp
@@ -106,15 +106,16 @@ void FAnimNode_MultiWayBlend::Update_AnyThread(const FAnimationUpdateContext& Co
 void FAnimNode_MultiWayBlend::Evaluate_AnyThread(FPoseContext& Output)
 {
 	DECLARE_SCOPE_HIERARCHICAL_COUNTER_ANIMNODE(Evaluate_AnyThread)
+
+	// this function may be reentrant when multiple MultiWayBlend nodes are chained together
+	// these scratch arrays are treated as stacks below
 	FMultiBlendData& BlendData = FMultiBlendData::Get();
 	TArray<FCompactPose, TInlineAllocator<8>>& SourcePoses = BlendData.SourcePoses;
-	TArray<float, TInlineAllocator<8>>& SourceWeights = BlendData.SourceWeights;
 	TArray<FBlendedCurve, TInlineAllocator<8>>& SourceCurves = BlendData.SourceCurves;
+	TArray<float, TInlineAllocator<8>>& SourceWeights = BlendData.SourceWeights;
 
-	SourcePoses.Reset();
-	SourceWeights.Reset();
-	SourceCurves.Reset();
-
+	const int32 SourcePosesInitialNum = SourcePoses.Num();
+	int32 SourcePosesAdded = 0;
 	if (ensure(Poses.Num() == CachedAlphas.Num()))
 	{
 		for (int32 PoseIndex = 0; PoseIndex < Poses.Num(); ++PoseIndex)
@@ -122,10 +123,11 @@ void FAnimNode_MultiWayBlend::Evaluate_AnyThread(FPoseContext& Output)
 			const float CurrentAlpha = CachedAlphas[PoseIndex];
 			if (CurrentAlpha > ZERO_ANIMWEIGHT_THRESH)
 			{
+				// evaluate input pose, potentially reentering this function and pushing/popping more poses
 				FPoseContext PoseContext(Output);
-				// total alpha shouldn't be zero
 				Poses[PoseIndex].Evaluate(PoseContext);
 
+				// push source pose data
 				FCompactPose& SourcePose = SourcePoses.AddDefaulted_GetRef();
 				SourcePose.MoveBonesFrom(PoseContext.Pose);
 
@@ -133,15 +135,28 @@ void FAnimNode_MultiWayBlend::Evaluate_AnyThread(FPoseContext& Output)
 				SourceCurve.MoveFrom(PoseContext.Curve);
 
 				SourceWeights.Add(CurrentAlpha);
+
+				++SourcePosesAdded;
 			}
 		}
 	}
 
-	if (SourcePoses.Num() > 0)
+	if (SourcePosesAdded > 0)
 	{
-		FAnimationRuntime::BlendPosesTogether(SourcePoses, SourceCurves, SourceWeights, Output.Pose, Output.Curve);
+		// obtain views onto the ends of our stacks
+		TArrayView<FCompactPose> SourcePosesView = MakeArrayView(&SourcePoses[SourcePosesInitialNum], SourcePosesAdded);
+		TArrayView<FBlendedCurve> SourceCurvesView = MakeArrayView(&SourceCurves[SourcePosesInitialNum], SourcePosesAdded);
+		TArrayView<float> SourceWeightsView = MakeArrayView(&SourceWeights[SourcePosesInitialNum], SourcePosesAdded);
+
+		FAnimationRuntime::BlendPosesTogether(SourcePosesView, SourceCurvesView, SourceWeightsView, Output.Pose, Output.Curve);
+		
 		// normalize rotation - some cases, where additive is applied less than 1, it will use non normalized rotation
 		Output.Pose.NormalizeRotations();
+		
+		// pop the poses we added
+		SourcePoses.SetNum(SourcePosesInitialNum, false);
+		SourceCurves.SetNum(SourcePosesInitialNum, false);
+		SourceWeights.SetNum(SourcePosesInitialNum, false);
 	}
 	else
 	{
