Crash fix on AD node with invalid set up

diff --git a/Engine/Source/Runtime/AnimGraphRuntime/Private/BoneControllers/AnimNode_AnimDynamics.cpp b/Engine/Source/Runtime/AnimGraphRuntime/Private/BoneControllers/AnimNode_AnimDynamics.cpp
--- a/Engine/Source/Runtime/AnimGraphRuntime/Private/BoneControllers/AnimNode_AnimDynamics.cpp
+++ b/Engine/Source/Runtime/AnimGraphRuntime/Private/BoneControllers/AnimNode_AnimDynamics.cpp
@@ -568,7 +568,7 @@ void FAnimNode_AnimDynamics::InitPhysics(FComponentSpacePoseContext& Output)
 				int32 ParentBoneIndex = BoneContainer.GetParentBoneIndex(ChainEnd.BoneIndex);
 
 				// Walk up the chain until we either find the top or hit the root bone
-				while(ParentBoneIndex != 0)
+				while(ParentBoneIndex > 0)
 				{
 					ChainBoneIndices.Add(ParentBoneIndex);
 					ChainBoneNames.Add(BoneContainer.GetReferenceSkeleton().GetBoneName(ParentBoneIndex));
