AnimNode_AnimDynamics: Build break on Mac

diff --git a/Engine/Source/Runtime/AnimGraphRuntime/Public/BoneControllers/AnimNode_AnimDynamics.h b/Engine/Source/Runtime/AnimGraphRuntime/Public/BoneControllers/AnimNode_AnimDynamics.h
--- a/Engine/Source/Runtime/AnimGraphRuntime/Public/BoneControllers/AnimNode_AnimDynamics.h
+++ b/Engine/Source/Runtime/AnimGraphRuntime/Public/BoneControllers/AnimNode_AnimDynamics.h
@@ -10,7 +10,7 @@
 #include "BonePose.h"
 #include "BoneControllers/AnimNode_SkeletalControlBase.h"
 #include "Animation/AnimInstanceProxy.h"
-#include "KismetAnimationTypes.h"
+#include "CommonAnimationTypes.h"
 #include "AnimNode_AnimDynamics.generated.h"
 
 class UAnimInstance;
