Copying //UE4/Dev-Physics to Dev-Main (//UE4/Dev-Main) [at] 7703071

diff --git a/Engine/Plugins/FX/Niagara/Source/NiagaraEditor/Private/NiagaraHlslTranslator.cpp b/Engine/Plugins/FX/Niagara/Source/NiagaraEditor/Private/NiagaraHlslTranslator.cpp
--- a/Engine/Plugins/FX/Niagara/Source/NiagaraEditor/Private/NiagaraHlslTranslator.cpp
+++ b/Engine/Plugins/FX/Niagara/Source/NiagaraEditor/Private/NiagaraHlslTranslator.cpp
@@ -15,6 +15,7 @@
 #include "NiagaraNodeWriteDataSet.h"
 #include "NiagaraNodeParameterMapGet.h"
 #include "NiagaraNodeParameterMapSet.h"
+#include "NiagaraNodeParameterMapFor.h"
 #include "NiagaraNodeCustomHlsl.h"
 #include "NiagaraNodeOp.h"
 #include "NiagaraNodeConvert.h"
@@ -63,6 +64,8 @@ DECLARE_CYCLE_STAT(TEXT("Niagara - HlslTranslator - CustomHLSL"), STAT_NiagaraEd
 DECLARE_CYCLE_STAT(TEXT("Niagara - HlslTranslator - FuncBody"), STAT_NiagaraEditor_Module_NiagaraHLSLTranslator_FuncBody, STATGROUP_NiagaraEditor);
 DECLARE_CYCLE_STAT(TEXT("Niagara - HlslTranslator - Output"), STAT_NiagaraEditor_HlslTranslator_Output, STATGROUP_NiagaraEditor);
 DECLARE_CYCLE_STAT(TEXT("Niagara - HlslTranslator - MapSet"), STAT_NiagaraEditor_HlslTranslator_MapSet, STATGROUP_NiagaraEditor);
+DECLARE_CYCLE_STAT(TEXT("Niagara - HlslTranslator - MapForBegin"), STAT_NiagaraEditor_HlslTranslator_MapForBegin, STATGROUP_NiagaraEditor);
+DECLARE_CYCLE_STAT(TEXT("Niagara - HlslTranslator - MapForEnd"), STAT_NiagaraEditor_HlslTranslator_MapForEnd, STATGROUP_NiagaraEditor);
 DECLARE_CYCLE_STAT(TEXT("Niagara - HlslTranslator - Operation"), STAT_NiagaraEditor_HlslTranslator_Operation, STATGROUP_NiagaraEditor);
 DECLARE_CYCLE_STAT(TEXT("Niagara - HlslTranslator - If"), STAT_NiagaraEditor_HlslTranslator_If, STATGROUP_NiagaraEditor);
 DECLARE_CYCLE_STAT(TEXT("Niagara - HlslTranslator - CompilePin"), STAT_NiagaraEditor_HlslTranslator_CompilePin, STATGROUP_NiagaraEditor);
@@ -591,6 +594,28 @@ FString FHlslNiagaraTranslator::BuildParameterMapHlslDefinitions(TArray<FNiagara
 		}
 	}
 
+	static const auto UseShaderStagesCVar = IConsoleManager::Get().FindConsoleVariable(TEXT("fx.UseShaderStages"));
+	if (UseShaderStagesCVar->GetInt() == 1)
+	{
+		// Add the attribute indices to the list of unique variables
+		TArray<FString> RegisterNames;
+		for (int32 UniqueVarIdx = 0; UniqueVarIdx < UniqueVariables.Num(); UniqueVarIdx++)
+		{
+			const FNiagaraVariable& NiagaraVariable = UniqueVariables[UniqueVarIdx];
+			if (FNiagaraParameterMapHistory::IsAttribute(NiagaraVariable))
+			{
+				const FString VariableName = GetSanitizedSymbolName(NiagaraVariable.GetName().ToString());
+				RegisterNames.Add(VariableName.Replace(PARAM_MAP_ATTRIBUTE_STR, PARAM_MAP_INDICES_STR));
+			}
+		}
+		for (const FString& RegisterName : RegisterNames)
+		{
+			FNiagaraVariable NiagaraVariable = FNiagaraVariable(FNiagaraTypeDefinition::GetIntDef(), *RegisterName);
+			UniqueVariables.AddUnique(NiagaraVariable);
+		}
+	}
+
+
 	TMap<FString, TArray<TPair<FString, FString>> > ParamStructNameToMembers;
 	TArray<FString> ParamStructNames;
 
@@ -1082,7 +1107,6 @@ const FNiagaraTranslateResults &FHlslNiagaraTranslator::Translate(const FNiagara
 				Error(ErrorMessage, nullptr, nullptr);
 			}
 		}
-
 		//Declare parameters.
 		//TODO: Separate Cbuffer for Global, System and Emitter parameters.
 		{
@@ -2006,12 +2030,8 @@ void FHlslNiagaraTranslator::DefineDataSetVariableWrites(FString &OutHlslOutput,
 {
 	//TODO Grab indices for data set writes (inc output) and do the write. Need to rewrite this for events interleaved..
 	OutHlslOutput += "\t{\n";
-	OutHlslOutput += "\tint TmpWriteIndex;\n";
-	if (UNiagaraScript::IsNonParticleScript(CompileOptions.TargetUsage))
-	{
-		OutHlslOutput += "\tbool bValid = true;\n";
-	}
-	else
+	bool bUsesAlive = false;
+	if (!UNiagaraScript::IsNonParticleScript(CompileOptions.TargetUsage))
 	{
 		FString DataSetName = Id.Name.ToString();
 		bool bHasPerParticleAliveSpawn = false;
@@ -2052,6 +2072,7 @@ void FHlslNiagaraTranslator::DefineDataSetVariableWrites(FString &OutHlslOutput,
 			// NOTE: Context.MapSpawn is copied to Context.MapUpdate before this point in the script, so we might
 			//       as well just keep it simple and check against MapUpdate only instead of redundantly branch.
 			OutHlslOutput += TEXT("\tbool bValid = Context.MapUpdate.") + DataSetName + TEXT(".Alive;\n");
+			bUsesAlive = true;
 		}
 		else if ((UNiagaraScript::IsParticleSpawnScript(CompileOptions.TargetUsage) && bHasPerParticleAliveSpawn) 
 			|| (UNiagaraScript::IsParticleUpdateScript(CompileOptions.TargetUsage) && bHasPerParticleAliveUpdate)
@@ -2059,15 +2080,19 @@ void FHlslNiagaraTranslator::DefineDataSetVariableWrites(FString &OutHlslOutput,
 		{
 			// Non-interpolated CPU spawn script
 			OutHlslOutput += TEXT("\tbool bValid = Context.Map.") + DataSetName + TEXT(".Alive;\n");
-		}
-		else
-		{
-			OutHlslOutput += "\tbool bValid = true;\n";
+			bUsesAlive = true;
 		}
 	}
 
-	// grab the current ouput index; currently pass true, but should use an arbitrary bool to determine whether write should happen or not
-	OutHlslOutput += "\tTmpWriteIndex = AcquireIndex(0, bValid);\n";
+	// grab the current ouput index to write datas 
+	if (bUsesAlive)
+	{
+		OutHlslOutput += "\tint TmpWriteIndex = OutputIndex(0, false, bValid);\n";
+	}
+	else
+	{
+		OutHlslOutput += "\tint TmpWriteIndex = OutputIndex(0, true, true);\n";
+	}
 
 	bool bNeedsPersistentIDs = CompileOptions.AdditionalDefines.Contains(TEXT("RequiresPersistentIDs"));
 	if (bNeedsPersistentIDs && DataSetIndex == 0)
@@ -2171,19 +2196,34 @@ void FHlslNiagaraTranslator::DefineDataSetVariableReads(FString &OutHlslOutput,
 
 		for (FNiagaraVariable &Var : ReadVars)
 		{
+			const FString VariableName = ContextName + GetSanitizedSymbolName(Var.GetName().ToString());
 			// If the NiagaraClearEachFrame value is set on the data set, we don't bother reading it in each frame as we know that it is is invalid. However,
 			// this is only used for the base data set. Other reads are potentially from events and are therefore perfectly valid.
 			if (DataSetIndex == 0 && Var.GetType().GetScriptStruct() != nullptr && Var.GetType().GetScriptStruct()->GetMetaData(TEXT("NiagaraClearEachFrame")).Equals(TEXT("true"), ESearchCase::IgnoreCase))
 			{
-				Fmt = ContextName + GetSanitizedSymbolName(Var.GetName().ToString()) + TEXT("{0} = {4};\n");
-			}
-			else if (CompilationTarget == ENiagaraSimTarget::GPUComputeSim)
-			{
-				Fmt = ContextName + GetSanitizedSymbolName(Var.GetName().ToString()) + TEXT("{0} = InputData{1}({2}, {3}, InstanceIdx);\n");
+				Fmt = VariableName + TEXT("{0} = {4};\n");
 			}
 			else
 			{
-				Fmt = ContextName + GetSanitizedSymbolName(Var.GetName().ToString()) + TEXT("{0} = InputData{1}({2}, {3});\n");
+				if (CompilationTarget == ENiagaraSimTarget::GPUComputeSim)
+				{
+					Fmt = VariableName + TEXT("{0} = InputData{1}({2}, {3}, InstanceIdx);\n");
+				}
+				else
+				{
+					Fmt = VariableName + TEXT("{0} = InputData{1}({2}, {3});\n");
+				}
+				static const auto UseShaderStagesCVar = IConsoleManager::Get().FindConsoleVariable(TEXT("fx.UseShaderStages"));
+				if (UseShaderStagesCVar->GetInt() == 1)
+				{
+					if (FNiagaraParameterMapHistory::IsAttribute(Var))
+					{
+						FString RegisterName = VariableName;
+						RegisterName.ReplaceInline(PARAM_MAP_ATTRIBUTE_STR, PARAM_MAP_INDICES_STR);
+						const int32 RegisterValue = Var.GetType().IsFloatPrimitive() ? FloatCounter : IntCounter;
+						VarReads += RegisterName + FString::Printf(TEXT(" = %d;\n"), RegisterValue);
+					}
+				}
 			}
 			GatherVariableForDataSetAccess(Var, Fmt, IntCounter, FloatCounter, DataSetIndex, TEXT(""), VarReads);
 		}
@@ -3239,6 +3279,19 @@ FString FHlslNiagaraTranslator::GetDataSetAccessSymbol(FNiagaraDataSetID DataSet
 	return Ret;
 }
 
+void FHlslNiagaraTranslator::ParameterMapForBegin(UNiagaraNodeParameterMapFor* ForNode, int32 IterationCount)
+{
+	SCOPE_CYCLE_COUNTER(STAT_NiagaraEditor_HlslTranslator_MapForBegin);
+
+	AddBodyChunk(TEXT(""), TEXT("for(int index = 0; index < {0}; ++index)\n\t{"), FNiagaraTypeDefinition::GetIntDef(), IterationCount, false, false);
+}
+
+void FHlslNiagaraTranslator::ParameterMapForEnd(UNiagaraNodeParameterMapFor* ForNode)
+{
+	SCOPE_CYCLE_COUNTER(STAT_NiagaraEditor_HlslTranslator_MapForEnd);
+
+	AddBodyChunk(TEXT(""), TEXT("}"), FNiagaraTypeDefinition::GetIntDef(), false, false);
+}
 
 void FHlslNiagaraTranslator::ParameterMapSet(UNiagaraNodeParameterMapSet* SetNode, TArray<int32>& Inputs, TArray<int32>& Outputs)
 {
@@ -3566,9 +3619,13 @@ bool FHlslNiagaraTranslator::ParameterMapRegisterExternalConstantNamespaceVariab
 				UniformIdx = ParamMapDefinedSystemVarsToUniformChunks.FindChecked(InVariable.GetName());
 				UniformChunk = ChunksByMode[(int32)ENiagaraCodeChunkMode::Uniform][UniformIdx];
 			}
-
-			//Add this separately as the same uniform can appear in the pre sim chunks more than once in different param maps.
-			MainPreSimulateChunks.AddUnique(FString::Printf(TEXT("%s.%s = %s;"), *ParameterMapInstanceName, *GetSanitizedSymbolName(VarName), *GetCodeAsSource(UniformChunk)));
+			static const auto UseShaderStagesCVar = IConsoleManager::Get().FindConsoleVariable(TEXT("fx.UseShaderStages"));
+			// Avoid overriding the register indices
+			if ((UseShaderStagesCVar->GetInt()==1 && !FNiagaraParameterMapHistory::IsInNamespace(InVariable, PARAM_MAP_INDICES_STR)) || UseShaderStagesCVar->GetInt()==0)
+			{
+				//Add this separately as the same uniform can appear in the pre sim chunks more than once in different param maps.
+				MainPreSimulateChunks.AddUnique(FString::Printf(TEXT("%s.%s = %s;"), *ParameterMapInstanceName, *GetSanitizedSymbolName(VarName), *GetCodeAsSource(UniformChunk)));
+			}
 		}
 		else if (bIsPerInstanceBulkSystemParam && !ExternalVariablesForBulkUsage.Contains(InVariable))
 		{
@@ -4760,6 +4817,16 @@ void FHlslNiagaraTranslator::RegisterFunctionCall(ENiagaraScriptUsage ScriptUsag
 				Error(FText::Format(LOCTEXT("FunctionCallMissingFunction", "Function call signature does not reference a function. Top-level module: {0} Source: {1}"), ModuleAlias ? FText::FromString(*ModuleAlias) : FText::FromString(TEXT("Unknown module")), FText::FromString(CompileOptions.FullName)), nullptr, nullptr);
 				return;
 			}
+			bool bIsModuleFunction = false;
+			bool bStageMinFilter = false;
+			bool bStageMaxFilter = false;
+			FString MinParam;
+			FString MaxParam;
+			FString MinParamSpawn;
+			FString MaxParamSpawn;
+
+			static const auto UseShaderStagesCVar = IConsoleManager::Get().FindConsoleVariable(TEXT("fx.UseShaderStages"));
+			const bool UseShaderStages = (UseShaderStagesCVar->GetInt() == 1) && CompilationTarget == ENiagaraSimTarget::GPUComputeSim;
 
 			//We've not compiled this function yet so compile it now.
 			EnterFunction(InName, OutSignature, Inputs, CallNodeId);
@@ -4807,6 +4874,24 @@ void FHlslNiagaraTranslator::RegisterFunctionCall(ENiagaraScriptUsage ScriptUsag
 								const FNiagaraVariable& AliasedVar = History.VariablesWithOriginalAliasesIntact[VarIdx];
 								bool bIsAliased = Var.GetName() != AliasedVar.GetName();
 
+								if (bIsAliased && UseShaderStages)
+								{
+									if (Var.GetName().ToString().Contains(TEXT("MinStage")))
+									{
+										MinParam = TEXT("Context.MapUpdate.Constants.Emitter.") + Var.GetName().ToString();
+										MinParamSpawn = TEXT("Context.MapSpawn.Constants.Emitter.") + Var.GetName().ToString();
+										bStageMinFilter = true;
+									}
+									if (Var.GetName().ToString().Contains(TEXT("MaxStage")))
+									{
+										MaxParam = TEXT("Context.MapUpdate.Constants.Emitter.") + Var.GetName().ToString();
+										MaxParamSpawn = TEXT("Context.MapSpawn.Constants.Emitter.") + Var.GetName().ToString();
+										bStageMaxFilter = true;
+									}
+								}
+
+								bIsModuleFunction = (bIsInTopLevelFunction && ParamMapPin != nullptr && UNiagaraScript::IsGPUScript(CompileOptions.TargetUsage));
+
 								// For non aliased values we resolve the defaults once at the top level since it's impossible to know which context they were actually used in, but
 								// for aliased values we check to see if they're used in the current context by resolving the alias and checking against the current resolved variable
 								// name since aliased values can only be resolved for reading in the correct context.
@@ -4834,7 +4919,7 @@ void FHlslNiagaraTranslator::RegisterFunctionCall(ENiagaraScriptUsage ScriptUsag
 				}
 			}
 
-			//Track the start of this funciton in the chunks so we can remove them after we grab the funcitons hlsl.
+			//Track the start of this function in the chunks so we can remove them after we grab the function hlsl.
 			int32 ChunkStart = CodeChunks.Num();
 			int32 ChunkStartsByMode[(int32)ENiagaraCodeChunkMode::Num];
 			for (int32 i = 0; i < (int32)ENiagaraCodeChunkMode::Num; ++i)
@@ -4858,11 +4943,37 @@ void FHlslNiagaraTranslator::RegisterFunctionCall(ENiagaraScriptUsage ScriptUsag
 				//Grab all the body chunks for this function.
 				FString FunctionDefStr;
 				FunctionDefStr.Reserve(256 * ChunksByMode[(int32)ENiagaraCodeChunkMode::Body].Num());
+
+				if (bIsModuleFunction)
+				{
+					if (UseShaderStages)
+					{
+						if (bStageMaxFilter && bStageMaxFilter)
+						{
+							FunctionDefStr += TEXT("if ((GCurrentPhase == 1 && ShaderStageIndex >= ") + MinParam + TEXT(" && ShaderStageIndex <= ") + MaxParam + TEXT(") || ") +
+								TEXT("(GCurrentPhase == 0 && ShaderStageIndex >= ") + MinParamSpawn + TEXT(" && ShaderStageIndex <= ") + MaxParamSpawn + TEXT(")") +
+								TEXT(")\n{\n");
+						}
+						else
+						{
+							FunctionDefStr += TEXT("if (ShaderStageIndex == 0)\n{\n");
+						}
+					}
+				}
+
 				for (int32 i = ChunkStartsByMode[(int32)ENiagaraCodeChunkMode::Body]; i < ChunksByMode[(int32)ENiagaraCodeChunkMode::Body].Num(); ++i)
 				{
 					FunctionDefStr += GetCode(ChunksByMode[(int32)ENiagaraCodeChunkMode::Body][i]);
 				}
 
+				if (bIsModuleFunction)
+				{
+					if (UseShaderStages)
+					{
+						FunctionDefStr += TEXT("}\n");
+					}
+				}
+
 				//Now remove all chunks for the function again.			
 				//This is super hacky. Should move chunks etc into a proper scoping system.
 
diff --git a/Engine/Plugins/FX/Niagara/Source/NiagaraEditor/Private/NiagaraHlslTranslator.h b/Engine/Plugins/FX/Niagara/Source/NiagaraEditor/Private/NiagaraHlslTranslator.h
--- a/Engine/Plugins/FX/Niagara/Source/NiagaraEditor/Private/NiagaraHlslTranslator.h
+++ b/Engine/Plugins/FX/Niagara/Source/NiagaraEditor/Private/NiagaraHlslTranslator.h
@@ -428,6 +428,9 @@ public:
 	virtual void ParameterMapGet(class UNiagaraNodeParameterMapGet* GetNode, TArray<int32>& Inputs, TArray<int32>& Outputs);
 	virtual void Emitter(class UNiagaraNodeEmitter* GetNode, TArray<int32>& Inputs, TArray<int32>& Outputs);
 
+	virtual void ParameterMapForBegin(class UNiagaraNodeParameterMapFor* ForNode, int32 IterationCount);
+	virtual void ParameterMapForEnd(class UNiagaraNodeParameterMapFor* ForNode);
+
 	void DefineInterpolatedParametersFunction(FString &HlslOutput);
 	void DefineDataSetReadFunction(FString &HlslOutput, TArray<FNiagaraDataSetID> &ReadDataSets);
 	void DefineDataSetWriteFunction(FString &HlslOutput, TArray<FNiagaraDataSetProperties> &WriteDataSets, TArray<int32>& WriteConditionVarIndices);
diff --git a/Engine/Plugins/FX/Niagara/Source/NiagaraEditor/Private/NiagaraNodeParameterMapFor.cpp b/Engine/Plugins/FX/Niagara/Source/NiagaraEditor/Private/NiagaraNodeParameterMapFor.cpp
new file mode 100644
--- /dev/null
+++ b/Engine/Plugins/FX/Niagara/Source/NiagaraEditor/Private/NiagaraNodeParameterMapFor.cpp
@@ -0,0 +1,62 @@
+// Copyright 1998-2019 Epic Games, Inc. All Rights Reserved.
+
+#include "NiagaraNodeParameterMapFor.h"
+#include "EdGraphSchema_Niagara.h"
+#include "NiagaraEditorUtilities.h"
+#include "SNiagaraGraphNodeConvert.h"
+#include "NiagaraHlslTranslator.h"
+#include "Templates/SharedPointer.h"
+#include "NiagaraGraph.h"
+#include "NiagaraConstants.h"
+#include "Framework/MultiBox/MultiBoxBuilder.h"
+#include "Widgets/Layout/SBox.h"
+#include "Widgets/Input/SEditableTextBox.h"
+#include "EdGraph/EdGraphNode.h"
+
+#include "ScopedTransaction.h"
+
+#define LOCTEXT_NAMESPACE "NiagaraNodeParameterMapFor" 
+
+UNiagaraNodeParameterMapFor::UNiagaraNodeParameterMapFor() : UNiagaraNodeParameterMapSet()
+{
+}
+
+void UNiagaraNodeParameterMapFor::AllocateDefaultPins()
+{
+	PinPendingRename = nullptr;
+	const UEdGraphSchema_Niagara* Schema = GetDefault<UEdGraphSchema_Niagara>();
+	CreatePin(EGPD_Input, Schema->TypeDefinitionToPinType(FNiagaraTypeDefinition::GetParameterMapDef()), *UNiagaraNodeParameterMapBase::SourcePinName.ToString());
+	CreatePin(EGPD_Input, Schema->TypeDefinitionToPinType(FNiagaraTypeDefinition::GetIntDef()), TEXT("Module.Iteration Count"));
+	CreatePin(EGPD_Output, Schema->TypeDefinitionToPinType(FNiagaraTypeDefinition::GetParameterMapDef()), *UNiagaraNodeParameterMapBase::DestPinName.ToString());
+	CreateAddPin(EGPD_Input);
+}
+
+void UNiagaraNodeParameterMapFor::Compile(FHlslNiagaraTranslator* Translator, TArray<int32>& Outputs)
+{
+	TArray<UEdGraphPin*> InputPins;
+	GetInputPins(InputPins);
+	if (Translator)
+	{
+		if (Translator->GetSimulationTarget() == ENiagaraSimTarget::GPUComputeSim)
+		{
+			const int32 IterationCount = Translator->CompilePin(InputPins[1]);
+			Translator->ParameterMapForBegin(this, IterationCount);
+
+			UNiagaraNodeParameterMapSet::Compile(Translator, Outputs);
+
+			Translator->ParameterMapForEnd(this);
+		}
+		else
+		{
+			UNiagaraNodeParameterMapSet::Compile(Translator, Outputs);
+			Translator->Warning(LOCTEXT("UnsupportedParamMapFor", "Parameter map for is not yet supported on cpu."), this, nullptr);
+		}
+	}
+}
+
+FText UNiagaraNodeParameterMapFor::GetNodeTitle(ENodeTitleType::Type TitleType) const
+{
+	return LOCTEXT("UNiagaraNodeParameterMapForName", "Map For");
+}
+
+#undef LOCTEXT_NAMESPACE
diff --git a/Engine/Plugins/FX/Niagara/Source/NiagaraEditor/Private/NiagaraParameterMapHistory.cpp b/Engine/Plugins/FX/Niagara/Source/NiagaraEditor/Private/NiagaraParameterMapHistory.cpp
--- a/Engine/Plugins/FX/Niagara/Source/NiagaraEditor/Private/NiagaraParameterMapHistory.cpp
+++ b/Engine/Plugins/FX/Niagara/Source/NiagaraEditor/Private/NiagaraParameterMapHistory.cpp
@@ -50,6 +50,12 @@ void FNiagaraParameterMapHistory::GetValidNamespacesForReading(ENiagaraScriptUsa
 	OutputNamespaces.Add(PARAM_MAP_SYSTEM_STR);
 	OutputNamespaces.Add(PARAM_MAP_EMITTER_STR);
 
+	static const auto UseShaderStagesCVar = IConsoleManager::Get().FindConsoleVariable(TEXT("fx.UseShaderStages"));
+	if (UseShaderStagesCVar->GetInt() == 1)
+	{
+		OutputNamespaces.Add(PARAM_MAP_INDICES_STR);
+	}
+
 	for (ENiagaraScriptUsage Usage : SupportedContexts)
 	{
 		if (UNiagaraScript::IsParticleScript(Usage))
@@ -90,6 +96,12 @@ bool FNiagaraParameterMapHistory::IsValidNamespaceForReading(ENiagaraScriptUsage
 	ConcernedNamespaces.Add(PARAM_MAP_SYSTEM_STR);
 	ConcernedNamespaces.Add(PARAM_MAP_EMITTER_STR);
 	ConcernedNamespaces.Add(PARAM_MAP_ATTRIBUTE_STR);
+
+	static const auto UseShaderStagesCVar = IConsoleManager::Get().FindConsoleVariable(TEXT("fx.UseShaderStages"));
+	if (UseShaderStagesCVar->GetInt() == 1)
+	{
+		ConcernedNamespaces.Add(PARAM_MAP_INDICES_STR);
+	}
 	
 	if (!Namespace.EndsWith(TEXT(".")))
 	{
@@ -566,6 +578,15 @@ bool FNiagaraParameterMapHistory::IsExternalConstantNamespace(const FNiagaraVari
 		return true;
 	}
 
+	static const auto UseShaderStagesCVar = IConsoleManager::Get().FindConsoleVariable(TEXT("fx.UseShaderStages"));
+	if (UseShaderStagesCVar->GetInt() == 1)
+	{
+		if (IsInNamespace(InVar, PARAM_MAP_INDICES_STR))
+		{
+			return true;
+		}
+	}
+
 	// Modules and functions need to act as if they are within the script types that they 
 	// say that they support rather than using their exact script type.
 	if (UNiagaraScript::IsStandaloneScript(InUsage))
diff --git a/Engine/Plugins/FX/Niagara/Source/NiagaraShader/Private/NiagaraShader.cpp b/Engine/Plugins/FX/Niagara/Source/NiagaraShader/Private/NiagaraShader.cpp
--- a/Engine/Plugins/FX/Niagara/Source/NiagaraShader/Private/NiagaraShader.cpp
+++ b/Engine/Plugins/FX/Niagara/Source/NiagaraShader/Private/NiagaraShader.cpp
@@ -299,6 +299,9 @@ FShaderCompileJob* FNiagaraShaderType::BeginCompileShader(
 	NewJob->Input.Environment.SetDefine(TEXT("GPU_SIMULATION"), 1);
 	NewJob->Input.Environment.SetDefine(TEXT("NIAGARA_MAX_GPU_SPAWN_INFOS"), NIAGARA_MAX_GPU_SPAWN_INFOS);
 	NewJob->Input.Environment.IncludeVirtualPathToContentsMap.Add(TEXT("/Engine/Generated/NiagaraEmitterInstance.ush"), Script->HlslOutput);
+
+	static const auto UseShaderStagesCVar = IConsoleManager::Get().FindConsoleVariable(TEXT("fx.UseShaderStages"));
+	NewJob->Input.Environment.SetDefine(TEXT("USE_SHADER_STAGES"), UseShaderStagesCVar->GetInt());
 	
 	AddReferencedUniformBufferIncludes(NewJob->Input.Environment, NewJob->Input.SourceFilePrefix, (EShaderPlatform)Target.Platform);
 	
@@ -1279,7 +1282,8 @@ void FNiagaraShader::BindParams(const FShaderParameterMap &ParameterMap)
 	InstanceCountsParam.Bind(ParameterMap, TEXT("InstanceCounts"));
 	ReadInstanceCountOffsetParam.Bind(ParameterMap, TEXT("ReadInstanceCountOffset"));
 	WriteInstanceCountOffsetParam.Bind(ParameterMap, TEXT("WriteInstanceCountOffset"));
-
+	
+	SimStartParam.Bind(ParameterMap, TEXT("SimStart"));
 	EmitterTickCounterParam.Bind(ParameterMap, TEXT("EmitterTickCounter"));
 	EmitterSpawnInfoOffsetsParam.Bind(ParameterMap, TEXT("EmitterSpawnInfoOffsets"));
 	EmitterSpawnInfoParamsParam.Bind(ParameterMap, TEXT("EmitterSpawnInfoParams"));
@@ -1290,6 +1294,8 @@ void FNiagaraShader::BindParams(const FShaderParameterMap &ParameterMap)
 
 	NumSpawnedInstancesParam.Bind(ParameterMap, TEXT("SpawnedInstances"));
 	UpdateStartInstanceParam.Bind(ParameterMap, TEXT("UpdateStartInstance"));
+	ShaderStageIndexParam.Bind(ParameterMap, TEXT("ShaderStageIndex"));
+	IterationInterfaceCount.Bind(ParameterMap, TEXT("IterationInterfaceCount"));
 
 	ComponentBufferSizeReadParam.Bind(ParameterMap, TEXT("ComponentBufferSizeRead"));
 	ComponentBufferSizeWriteParam.Bind(ParameterMap, TEXT("ComponentBufferSizeWrite"));
@@ -1348,12 +1354,15 @@ bool FNiagaraShader::Serialize(FArchive& Ar)
 	Ar << ReadInstanceCountOffsetParam;
 	Ar << WriteInstanceCountOffsetParam;
 
+	Ar << SimStartParam;
 	Ar << EmitterTickCounterParam;
 	Ar << EmitterSpawnInfoOffsetsParam;
 	Ar << EmitterSpawnInfoParamsParam;
 
 	Ar << NumSpawnedInstancesParam;
 	Ar << UpdateStartInstanceParam;
+	Ar << ShaderStageIndexParam;
+	Ar << IterationInterfaceCount;
 	Ar << ComponentBufferSizeReadParam;
 	Ar << ComponentBufferSizeWriteParam;
 
diff --git a/Engine/Plugins/FX/Niagara/Source/NiagaraShader/Public/NiagaraShader.h b/Engine/Plugins/FX/Niagara/Source/NiagaraShader/Public/NiagaraShader.h
--- a/Engine/Plugins/FX/Niagara/Source/NiagaraShader/Public/NiagaraShader.h
+++ b/Engine/Plugins/FX/Niagara/Source/NiagaraShader/Public/NiagaraShader.h
@@ -108,6 +108,7 @@ public:
 	FShaderUniformBufferParameter EmitterConstantBufferParam;
 	FShaderUniformBufferParameter DataInterfaceUniformBufferParam;
 	FShaderUniformBufferParameter ViewUniformBufferParam;
+	FShaderParameter SimStartParam;
 	FShaderParameter EmitterTickCounterParam;
 	FShaderParameter EmitterSpawnInfoOffsetsParam;
 	FShaderParameter EmitterSpawnInfoParamsParam;
@@ -116,6 +117,8 @@ public:
 	FShaderParameter CopyInstancesBeforeStartParam;
 	FShaderParameter NumSpawnedInstancesParam;
 	FShaderParameter UpdateStartInstanceParam;
+	FShaderParameter ShaderStageIndexParam;
+	FShaderParameter IterationInterfaceCount;
 	FShaderParameter ComponentBufferSizeReadParam;
 	FShaderParameter ComponentBufferSizeWriteParam;
 	FRWShaderParameter EventIntUAVParams[MAX_CONCURRENT_EVENT_DATASETS];
diff --git a/Engine/Plugins/Runtime/ApexDestruction/Source/ApexDestruction/Private/DestructibleComponent.cpp b/Engine/Plugins/Runtime/ApexDestruction/Source/ApexDestruction/Private/DestructibleComponent.cpp
--- a/Engine/Plugins/Runtime/ApexDestruction/Source/ApexDestruction/Private/DestructibleComponent.cpp
+++ b/Engine/Plugins/Runtime/ApexDestruction/Source/ApexDestruction/Private/DestructibleComponent.cpp
@@ -27,7 +27,7 @@
 #include "UObject/UObjectThreadContext.h"
 
 #if WITH_PHYSX
-#include "Physics/PhysicsGeometryPhysX.h"
+#include "Physics/PhysicsGeometry.h"
 #endif
 
 UDestructibleComponent::UDestructibleComponent(const FObjectInitializer& ObjectInitializer)
@@ -52,7 +52,7 @@ UDestructibleComponent::UDestructibleComponent(const FObjectInitializer& ObjectI
 
 	SetComponentSpaceTransformsDoubleBuffering(false);
 
-#if WITH_PHYSX
+#if WITH_PHYSX && PHYSICS_INTERFACE_PHYSX
 	// Get contact offset params
 	FBodySetupShapeIterator::GetContactOffsetParams(ContactOffsetFactor, MinContactOffset, MaxContactOffset);
 #endif //WITH_PHYSX
diff --git a/Engine/Plugins/ScriptPlugin/Source/ScriptGeneratorPlugin/Private/LuaScriptCodeGenerator.cpp b/Engine/Plugins/ScriptPlugin/Source/ScriptGeneratorPlugin/Private/LuaScriptCodeGenerator.cpp
--- a/Engine/Plugins/ScriptPlugin/Source/ScriptGeneratorPlugin/Private/LuaScriptCodeGenerator.cpp
+++ b/Engine/Plugins/ScriptPlugin/Source/ScriptGeneratorPlugin/Private/LuaScriptCodeGenerator.cpp
@@ -4,6 +4,8 @@
 #include "ScriptGeneratorLog.h"
 #include "UObject/ErrorException.h"
 #include "UObject/Class.h"
+#include "UObject/UnrealType.h"
+#include "Misc/Paths.h"
 
 // Supported structs
 static FName Name_Vector2D("Vector2D");
diff --git a/Engine/Plugins/ScriptPlugin/Source/ScriptGeneratorPlugin/Private/ScriptCodeGeneratorBase.cpp b/Engine/Plugins/ScriptPlugin/Source/ScriptGeneratorPlugin/Private/ScriptCodeGeneratorBase.cpp
--- a/Engine/Plugins/ScriptPlugin/Source/ScriptGeneratorPlugin/Private/ScriptCodeGeneratorBase.cpp
+++ b/Engine/Plugins/ScriptPlugin/Source/ScriptGeneratorPlugin/Private/ScriptCodeGeneratorBase.cpp
@@ -5,6 +5,7 @@
 #include "Misc/Paths.h"
 #include "UObject/Class.h"
 #include "ScriptGeneratorLog.h"
+#include "UObject/UnrealType.h"
 
 FScriptCodeGeneratorBase::FScriptCodeGeneratorBase(const FString& InRootLocalPath, const FString& InRootBuildPath, const FString& InOutputDirectory, const FString& InIncludeBase)
 {
diff --git a/Engine/Shaders/Private/StrandHairVertexFactory.ush b/Engine/Shaders/Private/StrandHairVertexFactory.ush
new file mode 100644
--- /dev/null
+++ b/Engine/Shaders/Private/StrandHairVertexFactory.ush
@@ -0,0 +1,431 @@
+// Copyright 1998-2019 Epic Games, Inc. All Rights Reserved.
+
+/*=============================================================================
+	StrandHairFactory.usf
+=============================================================================*/
+
+// TODO/ISSUES
+// [ ] Prev position is not correctly computed
+
+#include "VertexFactoryCommon.ush"
+#include "LocalVertexFactoryCommon.ush"
+
+#include "/Engine/Generated/UniformBuffers/PrecomputedLightingBuffer.ush"
+
+#if MANUAL_VERTEX_FETCH
+	Buffer<float4> VertexFetch_InstanceOriginBuffer;
+	Buffer<float4> VertexFetch_InstanceTransformBuffer;
+#endif
+
+// VF_STRAND_HAIR
+// Define used in certain shader like DeepShadow for running specific code (coverage computation for instance) 
+// when vertices are produced by the strand hair vertex factory 
+// This is set by compilation enviromenent of the vertex factory
+
+float StrandHairDensity;
+float StrandHairRadius;
+float StrandHairLength;
+float StrandHairRadiusAtDepth1;
+Buffer<float4> StrandHairPositionBuffer;
+Buffer<float4> StrandHairTangentBuffer;
+
+// ## A quad is made of the following topology and indexing:
+//
+//  0__2    4 
+//  | /    /|
+//  |/    /_|
+//  1    5  3  
+//
+//  ## Control point identification
+//  Here is an example of two consecutive strands. The control point:
+//   * .a=1 starting
+//   * .a=0 within
+//   * .a=2 ending
+//
+//   _O       .a == 1  Strand 0
+//  | / /|
+//  |/ /_|
+//   _O       .a == 0
+//  | / /|
+//  |/ /_|
+//   _O       .a == 0
+//  | / /|
+//  |/ /_|
+//    O       .a == 2
+//     
+//    
+//    O       .a == 1   Strand 1
+//  | / /|
+//  |/ /_|
+//    O       .a == 2
+//
+//   ...
+
+/**
+ * Per-vertex inputs from bound vertex buffers
+ */
+struct FVertexFactoryInput
+{
+#if VF_USE_PRIMITIVE_SCENE_DATA
+	uint PrimitiveId : ATTRIBUTE13;
+#endif	
+	uint VertexId : SV_VertexID;
+};
+
+#if RAYHITGROUPSHADER
+FVertexFactoryInput LoadVertexFactoryInput(uint TriangleIndex, int VertexIndex)
+{
+	// #hair_todo
+	FTriangleBaseAttributes Triangle = LoadTriangleBaseAttributes(TriangleIndex);
+
+	FVertexFactoryInput Input;
+	Input.Position = float4(Triangle.LocalPositions[VertexIndex], 1.0f);
+	Input.VertexId = Triangle.Indices[VertexIndex];
+#if VF_USE_PRIMITIVE_SCENE_DATA
+	Input.PrimitiveId = GetInstanceUserData();
+#endif
+	return Input;
+}
+#endif
+
+/** 
+ * Per-vertex inputs from bound vertex buffers.  Used by passes with a trimmed down position-only shader.
+ */
+struct FPositionOnlyVertexFactoryInput
+{
+#if VF_USE_PRIMITIVE_SCENE_DATA
+	uint PrimitiveId : ATTRIBUTE1;
+#endif
+	uint VertexId : SV_VertexID;
+};
+
+/** 
+ * Caches intermediates that would otherwise have to be computed multiple times.  Avoids relying on the compiler to optimize out redundant operations.
+ */
+struct FVertexFactoryIntermediates
+{
+	half3x3 TangentToLocal;
+	half3x3 TangentToWorld;
+	half TangentToWorldSign;
+	uint PrimitiveId;
+	uint   HairPrimitiveId;
+	float2 HairDimensions;
+	float2 HairUV;
+	float  HairSeed;
+	float  HairDensity;
+};
+
+uint GetPrimitiveId(FVertexFactoryInput Input)
+{
+#if VF_USE_PRIMITIVE_SCENE_DATA
+	return Input.PrimitiveId;
+#else
+	return 0;
+#endif
+}
+
+uint GetHairPrimitiveId(FVertexFactoryInput Input)
+{
+	return Input.VertexId / 6;
+}
+
+struct FControlPoint
+{
+	float3 Position;
+	float  WorldRadius;
+	float  UCoord;
+	float  WorldLength;
+	float  Seed;
+};
+
+FControlPoint GetVertexPosition(FVertexFactoryInput Input)
+{
+	const uint QuadIndex	= Input.VertexId % 6;
+	const uint BaseIndex	= Input.VertexId / 6;
+	const uint OffsetIndex	= QuadIndex == 0 || QuadIndex == 2 || QuadIndex == 4 ? 0 : 1;
+	const float4 PackedData = StrandHairPositionBuffer[BaseIndex + OffsetIndex];
+	
+	FControlPoint Output = (FControlPoint)0;
+
+	const float3 ControlPoint   = (PackedData.xyz);
+	uint PackedAlpha			= asuint(PackedData.a);
+
+	const uint ControlPointType =   PackedAlpha & 0x3;
+	Output.WorldRadius			= ((PackedAlpha & 0x00FC) >> 2) / 63.f * StrandHairRadius;
+	Output.UCoord				= ((PackedAlpha & 0xFF00) >> 8) / 255.f;
+	Output.WorldLength			= ((PackedAlpha & 0xFF0000) >> 16) / 255.f * StrandHairLength;
+	Output.Seed					= ((PackedAlpha & 0xFF000000) >> 24) / 255.f;
+	
+	// Create a degenerated quad the end of each strand to cut between each strands
+	// #hair_todo: This is not efficient for short strand like fur for instance. Need to revisit that at some point
+	const bool bIsInvalidQuad = (ControlPointType == 1 && OffsetIndex == 1) || (ControlPointType == 2 && OffsetIndex == 0);
+	Output.Position = bIsInvalidQuad ? float3(1.#INF, 1.#INF, 1.#INF) : ControlPoint;
+
+	return Output;
+}
+
+float GetWorldStrandRadius(FVertexFactoryInput Input)
+{
+	return GetVertexPosition(Input).WorldRadius;
+}
+
+float GetUCoord(FVertexFactoryInput Input)
+{
+	return GetVertexPosition(Input).UCoord;
+}
+
+float GetWorldStrandLength(FVertexFactoryInput Input)
+{
+	return GetVertexPosition(Input).WorldLength;
+}
+
+float GetHairSeed(FVertexFactoryInput Input)
+{
+	return GetVertexPosition(Input).Seed;
+}
+
+/** Converts from vertex factory specific interpolants to a FMaterialPixelParameters, which is used by material inputs. */
+FMaterialPixelParameters GetMaterialPixelParameters(FVertexFactoryInterpolantsVSToPS Interpolants, float4 SvPosition)
+{
+	// GetMaterialPixelParameters is responsible for fully initializing the result
+	FMaterialPixelParameters Result = MakeInitializedMaterialPixelParameters();
+
+	half3 TangentToWorld0 = GetTangentToWorld0(Interpolants).xyz;
+	half4 TangentToWorld2 = GetTangentToWorld2(Interpolants);
+	Result.UnMirrored = TangentToWorld2.w;
+
+	// Required for previewing materials that use ParticleColor
+	Result.Particle.Color = half4(1,1,1,1);
+
+	Result.TangentToWorld = AssembleTangentToWorld( TangentToWorld0, TangentToWorld2 );
+#if USE_WORLDVERTEXNORMAL_CENTER_INTERPOLATION
+	Result.WorldVertexNormal_Center = Interpolants.TangentToWorld2_Center.xyz;
+#endif
+
+	Result.TwoSidedSign = 1;
+	Result.PrimitiveId = GetPrimitiveId(Interpolants);
+	Result.HairPrimitiveId = Interpolants.HairPrimitiveId;
+	Result.HairDimensions = Interpolants.HairDimensions;
+	Result.HairUV = Interpolants.HairUV;
+	Result.HairSeed = Interpolants.HairSeed;
+
+#if COMPILER_SUPPORTS_TO_SCALAR_MEMORY
+	Result.PrimitiveId = ToScalarMemory(Result.PrimitiveId);
+#endif
+
+	return Result;
+}
+
+half3x3 CalcTangentToWorldNoScale(FVertexFactoryIntermediates Intermediates, half3x3 TangentToLocal)
+{
+	half3x3 LocalToWorld = GetLocalToWorld3x3(Intermediates.PrimitiveId);
+	half3 InvScale = GetPrimitiveData(Intermediates.PrimitiveId).InvNonUniformScaleAndDeterminantSign.xyz;
+	LocalToWorld[0] *= InvScale.x;
+	LocalToWorld[1] *= InvScale.y;
+	LocalToWorld[2] *= InvScale.z;
+	return mul(TangentToLocal, LocalToWorld); 
+}
+
+/** Converts from vertex factory specific input to a FMaterialVertexParameters, which is used by vertex shader material inputs. */
+FMaterialVertexParameters GetMaterialVertexParameters(FVertexFactoryInput Input, FVertexFactoryIntermediates Intermediates, float3 WorldPosition, half3x3 TangentToLocal)
+{
+	FMaterialVertexParameters Result = (FMaterialVertexParameters)0;
+	Result.WorldPosition = WorldPosition;
+
+	// does not handle instancing!
+	Result.TangentToWorld = Intermediates.TangentToWorld;
+	
+	Result.PrevFrameLocalToWorld = GetPrimitiveData(Intermediates.PrimitiveId).PreviousLocalToWorld;
+
+	Result.PreSkinnedPosition = GetVertexPosition(Input).Position;
+	Result.PreSkinnedNormal = TangentToLocal[2];
+	Result.PrimitiveId = Intermediates.PrimitiveId;
+	return Result;
+}
+
+float4 CalcWorldPosition(float3 Position, uint PrimitiveId)
+{
+	return TransformLocalToTranslatedWorld(Position, PrimitiveId);
+}
+
+half3x3 CalcTangentToLocal(FVertexFactoryInput Input, out float TangentSign)
+{
+	const uint QuadIndex = Input.VertexId % 6;
+	const uint BaseIndex = Input.VertexId / 6;
+	const uint OffsetIndex = QuadIndex == 0 || QuadIndex == 2 || QuadIndex == 4 ? 0 : 1;
+	const half3 TangentInputX = StrandHairTangentBuffer[(BaseIndex + OffsetIndex) * 2 + 0].xyz;
+	const half4 TangentInputZ = StrandHairTangentBuffer[(BaseIndex + OffsetIndex) * 2 + 1];
+
+	half3 TangentX = TangentBias(TangentInputX);
+	half4 TangentZ = TangentBias(TangentInputZ);
+	TangentSign = TangentZ.w;
+	
+	// derive the binormal by getting the cross product of the normal and tangent
+	half3 TangentY = cross(TangentZ.xyz, TangentX) * TangentZ.w;
+	
+	// Recalculate TangentX off of the other two vectors
+	// This corrects quantization error since TangentX was passed in as a quantized vertex input
+	// The error shows up most in specular off of a mesh with a smoothed UV seam (normal is smooth, but tangents vary across the seam)
+	half3x3 Result;
+	Result[0] = cross(TangentY, TangentZ.xyz) * TangentZ.w;
+	Result[1] = TangentY;
+	Result[2] = TangentZ.xyz;
+
+	return Result;
+}
+
+half3x3 CalcTangentToWorld(FVertexFactoryIntermediates Intermediates, half3x3 TangentToLocal)
+{
+	half3x3 TangentToWorld = CalcTangentToWorldNoScale(Intermediates, TangentToLocal);
+	return TangentToWorld;
+}
+
+float3 GetTangent(FVertexFactoryInput Input)
+{
+	FVertexFactoryIntermediates Intermediates;
+	Intermediates.PrimitiveId = GetPrimitiveId(Input);
+	
+	float TangentSign;
+	Intermediates.TangentToLocal = CalcTangentToLocal(Input, TangentSign);
+	Intermediates.TangentToWorld = CalcTangentToWorld(Intermediates,Intermediates.TangentToLocal);
+	Intermediates.TangentToWorldSign = TangentSign * GetPrimitiveData(Intermediates.PrimitiveId).InvNonUniformScaleAndDeterminantSign.w;
+
+	// Hair shader expec the tangent to be stored in place of the normal. This is what the StrandHairComponent is filling in:
+	// [0]: Real normal		| Should be Tangent
+	// [1]: Bitangent		| Should be BiTangent
+	// [2]: Real Tangent	| Should be Normal
+	return Intermediates.TangentToWorld[2].xyz;
+}
+
+FVertexFactoryIntermediates GetVertexFactoryIntermediates(FVertexFactoryInput Input)
+{
+	FVertexFactoryIntermediates Intermediates;
+	Intermediates.PrimitiveId = GetPrimitiveId(Input);
+	Intermediates.HairPrimitiveId = GetHairPrimitiveId(Input);
+	Intermediates.HairDimensions = float2(GetWorldStrandLength(Input), GetWorldStrandRadius(Input));
+	Intermediates.HairUV = float2(GetUCoord(Input), 0.5f);
+	Intermediates.HairDensity = StrandHairDensity;
+	Intermediates.HairSeed = GetHairSeed(Input);
+
+	float TangentSign;
+	Intermediates.TangentToLocal = CalcTangentToLocal(Input, TangentSign);
+	Intermediates.TangentToWorld = CalcTangentToWorld(Intermediates,Intermediates.TangentToLocal);
+	Intermediates.TangentToWorldSign = TangentSign * GetPrimitiveData(Intermediates.PrimitiveId).InvNonUniformScaleAndDeterminantSign.w;	
+	return Intermediates;
+}
+
+/**
+* Get the 3x3 tangent basis vectors for this vertex factory
+* this vertex factory will calculate the binormal on-the-fly
+*
+* @param Input - vertex input stream structure
+* @return 3x3 matrix
+*/
+half3x3 VertexFactoryGetTangentToLocal( FVertexFactoryInput Input, FVertexFactoryIntermediates Intermediates )
+{
+	return Intermediates.TangentToLocal;
+}
+
+float4 ComputeViewAlignedWorldPosition(FVertexFactoryInput Input, float3 WorldTangent, float4 WorldPosition, float WorldStrandRadius)
+{
+	const uint LocalVertexId = Input.VertexId % 6;
+	const bool bIsLeft = LocalVertexId == 0 || LocalVertexId == 1 || LocalVertexId == 5;
+	
+	// #hair_todo
+	// Use either the min strand radius computated within the strand factory, or a custom one, provided 
+	// (in a hacky way) throught the View variables. This is used by the DeepShadow pass, which inject 
+	// its data this way.
+	const bool  bUseCustomStrandRadius = ResolvedView.TranslatedWorldCameraOrigin.x == 1337.f;
+	const bool  bIsOrtho			   = bUseCustomStrandRadius && ResolvedView.TranslatedWorldCameraOrigin.y > 0;
+	const float CustomMinStrandRadius  = ResolvedView.TranslatedWorldCameraOrigin.z;
+	const float DistanceToCamera	   = length(ResolvedView.TranslatedWorldCameraOrigin - WorldPosition.xyz);
+
+	// Minimal radius to snap the strand to a sample/pixel center (to avoid aliasing)
+	const float StrandRadiusAtDepth1 = bUseCustomStrandRadius ? CustomMinStrandRadius : StrandHairRadiusAtDepth1;
+	const float MinStrandHairRadius  = bIsOrtho ? StrandRadiusAtDepth1 : (DistanceToCamera * StrandRadiusAtDepth1);
+	const float3 ViewDir = -ResolvedView.ViewForward;
+	const float3 Right = normalize(cross(WorldTangent, ViewDir));
+	const float3 OutWorldPosition = WorldPosition.xyz + (bIsLeft ? -Right : Right) * max(WorldStrandRadius, MinStrandHairRadius);
+	
+	return float4(OutWorldPosition, 1);
+}
+
+// @return translated world position
+float4 VertexFactoryGetWorldPosition(FVertexFactoryInput Input, FVertexFactoryIntermediates Intermediates)
+{
+	FControlPoint ControlPoint = GetVertexPosition(Input);
+	const float3 VertexPosition = ControlPoint.Position;
+	const float4 WorldPosition = CalcWorldPosition(VertexPosition, Intermediates.PrimitiveId);
+
+	// Hair shader expect the WorldNormal to be the tangent vector
+	const float3 WorldTangent = GetTangent(Input);
+	return ComputeViewAlignedWorldPosition(Input, WorldTangent, WorldPosition, ControlPoint.WorldRadius);
+}
+
+float4 VertexFactoryGetRasterizedWorldPosition(FVertexFactoryInput Input, FVertexFactoryIntermediates Intermediates, float4 InWorldPosition)
+{
+	return InWorldPosition;
+}
+
+float3 VertexFactoryGetPositionForVertexLighting(FVertexFactoryInput Input, FVertexFactoryIntermediates Intermediates, float3 TranslatedWorldPosition)
+{
+	return TranslatedWorldPosition;
+}
+
+FVertexFactoryInterpolantsVSToPS VertexFactoryGetInterpolantsVSToPS(FVertexFactoryInput Input, FVertexFactoryIntermediates Intermediates, FMaterialVertexParameters VertexParameters)
+{
+	FVertexFactoryInterpolantsVSToPS Interpolants;
+
+	// Initialize the whole struct to 0
+	// Really only the last two components of the packed UVs have the opportunity to be uninitialized
+	Interpolants = (FVertexFactoryInterpolantsVSToPS)0;
+
+	SetTangents(Interpolants, Intermediates.TangentToWorld[0], Intermediates.TangentToWorld[2], Intermediates.TangentToWorldSign);
+
+#if INSTANCED_STEREO
+	Interpolants.EyeIndex = 0;
+#endif
+	
+	Interpolants.HairPrimitiveId = GetHairPrimitiveId(Input);
+	Interpolants.HairDimensions = float2(GetWorldStrandLength(Input), GetWorldStrandRadius(Input));
+	Interpolants.HairUV = float2(GetUCoord(Input), 0.5f);
+	Interpolants.HairSeed = GetHairSeed(Input);
+	SetPrimitiveId(Interpolants, Intermediates.PrimitiveId);
+
+	return Interpolants;
+}
+
+/** for depth-only pass */
+float4 VertexFactoryGetWorldPosition(FPositionOnlyVertexFactoryInput Input)
+{
+	const FControlPoint ControlPoint = GetVertexPosition(Input);
+	const float3 Position = ControlPoint.Position;
+	const uint PrimitiveId = GetPrimitiveId(Input);
+
+	const float4 WorldPosition = CalcWorldPosition(Position, PrimitiveId);
+    const float3 WorldTangent = GetTangent(Input);
+	return ComputeViewAlignedWorldPosition(Input, WorldTangent, WorldPosition, ControlPoint.WorldRadius);
+}
+
+// @return previous translated world position
+float4 VertexFactoryGetPreviousWorldPosition(FVertexFactoryInput Input, FVertexFactoryIntermediates Intermediates)
+{
+	float4x4 PreviousLocalToWorldTranslated = GetPrimitiveData(Intermediates.PrimitiveId).PreviousLocalToWorld;
+	PreviousLocalToWorldTranslated[3][0] += ResolvedView.PrevPreViewTranslation.x;
+	PreviousLocalToWorldTranslated[3][1] += ResolvedView.PrevPreViewTranslation.y;
+	PreviousLocalToWorldTranslated[3][2] += ResolvedView.PrevPreViewTranslation.z;
+
+	const float3 VertexPosition = GetVertexPosition(Input).Position;
+	return mul(float4(VertexPosition,1), PreviousLocalToWorldTranslated);
+}
+
+float4 VertexFactoryGetTranslatedPrimitiveVolumeBounds(FVertexFactoryInterpolantsVSToPS Interpolants)
+{
+	float4 ObjectWorldPositionAndRadius = GetPrimitiveData(GetPrimitiveId(Interpolants)).ObjectWorldPositionAndRadius;
+	return float4(ObjectWorldPositionAndRadius.xyz + ResolvedView.PreViewTranslation.xyz, ObjectWorldPositionAndRadius.w);
+}
+
+uint VertexFactoryGetPrimitiveId(FVertexFactoryInterpolantsVSToPS Interpolants)
+{
+	return GetPrimitiveId(Interpolants);
+}
diff --git a/Engine/Source/Developer/MeshMergeUtilities/Private/MeshMergeHelpers.cpp b/Engine/Source/Developer/MeshMergeUtilities/Private/MeshMergeHelpers.cpp
--- a/Engine/Source/Developer/MeshMergeUtilities/Private/MeshMergeHelpers.cpp
+++ b/Engine/Source/Developer/MeshMergeUtilities/Private/MeshMergeHelpers.cpp
@@ -987,6 +987,9 @@ void FMeshMergeHelpers::ExtractPhysicsGeometry(UBodySetup* InBodySetup, const FT
 	{
 		Elem.SetConvexMesh(nullptr);
 		Elem.SetMirroredConvexMesh(nullptr);
+#if WITH_CHAOS
+		Elem.ResetChaosConvexMesh();
+#endif
 	}
 
 	// Transform geometry to world space
diff --git a/Engine/Source/Programs/UnrealBuildTool/Configuration/ModuleRules.cs b/Engine/Source/Programs/UnrealBuildTool/Configuration/ModuleRules.cs
--- a/Engine/Source/Programs/UnrealBuildTool/Configuration/ModuleRules.cs
+++ b/Engine/Source/Programs/UnrealBuildTool/Configuration/ModuleRules.cs
@@ -934,7 +934,7 @@ namespace UnrealBuildTool
 			PublicIncludePathModuleNames.Add("PhysicsCore");
 			PublicDependencyModuleNames.Add("PhysicsCore");
 
-			bool bUseNonPhysXInterface = Target.bUseChaos == true || Target.bCompileImmediatePhysics == true;
+			bool bUseNonPhysXInterface = Target.bUseChaos == true;
 
             // 
             if (Target.bCompileChaos == true || Target.bUseChaos == true)
@@ -988,8 +988,6 @@ namespace UnrealBuildTool
 				//
 				PublicDefinitions.Add("WITH_CHAOS_NEEDS_TO_BE_FIXED=0");
 
-				PublicDefinitions.Add("PHYSICS_INTERFACE_LLIMMEDIATE=0");
-
 				if (Target.bCompilePhysX)
 				{
 					PublicDefinitions.Add("PHYSICS_INTERFACE_PHYSX=1");
@@ -1081,16 +1079,6 @@ namespace UnrealBuildTool
 					PublicDefinitions.Add("WITH_CHAOS=0");
 					PublicDefinitions.Add("WITH_CHAOS_NEEDS_TO_BE_FIXED=0");
 				}
-
-				if(Target.bCompileImmediatePhysics)
-				{
-					PublicDefinitions.Add("PHYSICS_INTERFACE_LLIMMEDIATE=1");
-					PublicDependencyModuleNames.Add("PhysX");
-				}
-				else
-				{
-					PublicDefinitions.Add("PHYSICS_INTERFACE_LLIMMEDIATE=0");
-				}
 			}
 
 			if(Target.bCustomSceneQueryStructure)
diff --git a/Engine/Source/Programs/UnrealBuildTool/Configuration/TargetRules.cs b/Engine/Source/Programs/UnrealBuildTool/Configuration/TargetRules.cs
--- a/Engine/Source/Programs/UnrealBuildTool/Configuration/TargetRules.cs
+++ b/Engine/Source/Programs/UnrealBuildTool/Configuration/TargetRules.cs
@@ -325,13 +325,13 @@ namespace UnrealBuildTool
 		/// Whether to use the Chaos physics interface. This overrides the physx flags to disable APEX and NvCloth
 		/// </summary>
 		[RequiresUniqueBuildEnvironment]
-        public bool bUseChaos = false;
+		public bool bUseChaos = false;
 
 		/// <summary>
-		/// Whether to include the immediate mode physics interface. This overrides the physx flags to disable APEX and NvCloth
+		/// Whether to compile in checked chaos features for debugging
 		/// </summary>
 		[RequiresUniqueBuildEnvironment]
-		public bool bCompileImmediatePhysics = false;
+		public bool bUseChaosChecked = false;
 
 		/// <summary>
 		/// Whether scene query acceleration is done by UE4. The physx scene query structure is still created, but we do not use it.
@@ -1785,14 +1785,14 @@ namespace UnrealBuildTool
 			get { return Inner.bCompileChaos; }
 		}
 
-        public bool bUseChaos
-        {
-            get { return Inner.bUseChaos; }
-        }
+		public bool bUseChaos
+		{
+			get { return Inner.bUseChaos; }
+		}
 
-        public bool bCompileImmediatePhysics
+		public bool bUseChaosChecked
 		{
-			get { return Inner.bCompileImmediatePhysics; }
+			get { return Inner.bUseChaosChecked; }
 		}
 
 		public bool bCustomSceneQueryStructure
diff --git a/Engine/Source/Runtime/AnimGraphRuntime/Private/BoneControllers/AnimNode_RigidBody.cpp b/Engine/Source/Runtime/AnimGraphRuntime/Private/BoneControllers/AnimNode_RigidBody.cpp
--- a/Engine/Source/Runtime/AnimGraphRuntime/Private/BoneControllers/AnimNode_RigidBody.cpp
+++ b/Engine/Source/Runtime/AnimGraphRuntime/Private/BoneControllers/AnimNode_RigidBody.cpp
@@ -8,15 +8,16 @@
 #include "PhysicsEngine/PhysicsConstraintTemplate.h"
 #include "GameFramework/PawnMovementComponent.h"
 #include "Physics/PhysicsInterfaceCore.h"
-#include "Physics/ImmediatePhysics/ImmediatePhysicsSimulation.h"
 #include "Physics/ImmediatePhysics/ImmediatePhysicsActorHandle.h"
+#include "Physics/ImmediatePhysics/ImmediatePhysicsSimulation.h"
+#include "Physics/ImmediatePhysics/ImmediatePhysicsStats.h"
 #include "PhysicsEngine/PhysicsSettings.h"
 #include "Logging/MessageLog.h"
 
 //#pragma optimize("", off)
 
 /////////////////////////////////////////////////////
-// FAnimNode_Ragdoll
+// FAnimNode_RigidBody
 
 #define LOCTEXT_NAMESPACE "ImmediatePhysics"
 
@@ -216,7 +217,8 @@ void FAnimNode_RigidBody::InitializeNewBodyTransformsDuringSimulation(FComponent
 }
 
 DECLARE_CYCLE_STAT(TEXT("RigidBody_Eval"), STAT_RigidBody_Eval, STATGROUP_Anim);
-DECLARE_CYCLE_STAT(TEXT("FAnimNode_Ragdoll::EvaluateSkeletalControl_AnyThread"), STAT_ImmediateEvaluateSkeletalControl, STATGROUP_ImmediatePhysics);
+
+DECLARE_CYCLE_STAT(TEXT("FAnimNode_RigidBody::EvaluateSkeletalControl_AnyThread"), STAT_ImmediateEvaluateSkeletalControl, STATGROUP_ImmediatePhysics);
 
 void FAnimNode_RigidBody::EvaluateSkeletalControl_AnyThread(FComponentSpacePoseContext& Output, TArray<FBoneTransform>& OutBoneTransforms)
 {
@@ -224,7 +226,7 @@ void FAnimNode_RigidBody::EvaluateSkeletalControl_AnyThread(FComponentSpacePoseC
 	SCOPE_CYCLE_COUNTER(STAT_RigidBody_Eval);
 	FScopeCycleCounterUObject AdditionalScope(UsePhysicsAsset, GET_STATID(STAT_RigidBody_Eval));
 	SCOPE_CYCLE_COUNTER(STAT_ImmediateEvaluateSkeletalControl);
-	//SCOPED_NAMED_EVENT_TEXT("FAnimNode_Ragdoll::EvaluateSkeletalControl_AnyThread", FColor::Magenta);
+	//SCOPED_NAMED_EVENT_TEXT("FAnimNode_RigidBody::EvaluateSkeletalControl_AnyThread", FColor::Magenta);
 
 	// Update our eval counter, and decide whether we need to reset simulated bodies, if our anim instance hasn't updated in a while.
 	if(EvalCounter.HasEverBeenUpdated())
@@ -593,7 +595,7 @@ void FAnimNode_RigidBody::InitPhysics(const UAnimInstance* InAnimInstance)
 	SkeletonBoneIndexToBodyIndex.Reset(NumSkeletonBones);
 	SkeletonBoneIndexToBodyIndex.Init(INDEX_NONE, NumSkeletonBones);
 
-	PreviousTransform = InAnimInstance->GetSkelMeshComponent()->GetComponentToWorld();
+	PreviousTransform = SkeletalMeshComp->GetComponentToWorld();
 
 	if (UPhysicsSettings* Settings = UPhysicsSettings::Get())
 	{
@@ -616,7 +618,8 @@ void FAnimNode_RigidBody::InitPhysics(const UAnimInstance* InAnimInstance)
 		BodyAnimData.Reset(NumBodies);
 		BodyAnimData.AddDefaulted(NumBodies);
 		TotalMass = 0.f;
-		
+
+		// Instantiate a FBodyInstance/FConstraintInstance set that will be cloned into the Immediate Physics sim.
 		TArray<FBodyInstance*> HighLevelBodyInstances;
 		TArray<FConstraintInstance*> HighLevelConstraintInstances;
 		SkeletalMeshComp->InstantiatePhysicsAssetRefPose(*UsePhysicsAsset, SimulationSpace == ESimulationSpace::WorldSpace ? SkeletalMeshComp->GetComponentToWorld().GetScale3D() : FVector(1.f), HighLevelBodyInstances, HighLevelConstraintInstances);
@@ -643,63 +646,37 @@ void FAnimNode_RigidBody::InitPhysics(const UAnimInstance* InAnimInstance)
 			}
 		}
 
-		auto InsertBodiesHelper = [&](bool bSimulatedBodies)
+		// Create the immediate physics bodies
+		for (FBoneIndexType InsertBone : InsertionOrder)
 		{
-			for (FBoneIndexType InsertBone : InsertionOrder)
+			if (FBodyInstance* BodyInstance = BodiesSorted[InsertBone])
 			{
-				if (FBodyInstance* BodyInstance = BodiesSorted[InsertBone])
-				{
-					UBodySetup* BodySetup = UsePhysicsAsset->SkeletalBodySetups[BodyInstance->InstanceBodyIndex];
-					const bool bKinematic = BodySetup->PhysicsType != EPhysicsType::PhysType_Simulated;
-					const FTransform& LastTransform = SkeletalMeshComp->GetBoneTransform(InsertBone);	//This is out of date, but will still give our bodies an initial setup that matches the constraints (TODO: use refpose)
+				UBodySetup* BodySetup = UsePhysicsAsset->SkeletalBodySetups[BodyInstance->InstanceBodyIndex];
 
-					ImmediatePhysics::FActorHandle* NewBodyHandle = nullptr;
-					if (bSimulatedBodies && !bKinematic)
+				bool bSimulated = (BodySetup->PhysicsType == EPhysicsType::PhysType_Simulated);
+				ImmediatePhysics::EActorType ActorType = bSimulated ?  ImmediatePhysics::EActorType::DynamicActor : ImmediatePhysics::EActorType::KinematicActor;
+				ImmediatePhysics::FActorHandle* NewBodyHandle = PhysicsSimulation->CreateActor(ActorType, BodyInstance, BodyInstance->GetUnrealWorldTransform());
+				if (NewBodyHandle)
+				{
+					if (bSimulated)
 					{
-#if WITH_PHYSX && PHYSICS_INTERFACE_PHYSX
-						PxRigidActor* PRigidActor = FPhysicsInterface::GetPxRigidActor_AssumesLocked(BodyInstance->GetPhysicsActorHandle());
-						if(PxRigidDynamic* PDynamic = PRigidActor->is<PxRigidDynamic>())
-						{
-							NewBodyHandle = PhysicsSimulation->CreateDynamicActor(PDynamic, LastTransform);
-						checkSlow(NewBodyHandle);
 						const float InvMass = NewBodyHandle->GetInverseMass();
 						TotalMass += InvMass > 0.f ? 1.f / InvMass : 0.f;
-						}
-#endif
 					}
-					else if (bSimulatedBodies && bKinematic)
+					const int32 BodyIndex = Bodies.Add(NewBodyHandle);
+					const int32 SkeletonBoneIndex = MeshToSkeletonBoneIndex[InsertBone];
+					SkeletonBoneIndexToBodyIndex[SkeletonBoneIndex] = BodyIndex;
+					BodyAnimData[BodyIndex].bIsSimulated = bSimulated;
+					NamesToHandles.Add(BodySetup->BoneName, NewBodyHandle);
+					BodyIndexToActorHandle[BodyInstance->InstanceBodyIndex] = NewBodyHandle;
+
+					if (BodySetup->CollisionReponse == EBodyCollisionResponse::BodyCollision_Disabled)
 					{
-#if WITH_PHYSX && PHYSICS_INTERFACE_PHYSX
-						PxRigidActor* PRigidActor = FPhysicsInterface::GetPxRigidActor_AssumesLocked(BodyInstance->GetPhysicsActorHandle());
-						if(PxRigidBody* PRigidBody = PRigidActor->is<PxRigidBody>())
-						{
-							NewBodyHandle = PhysicsSimulation->CreateKinematicActor(PRigidBody, LastTransform);
-						}
-#endif
-					}
-
-					if (NewBodyHandle)
-					{
-						const int32 BodyIndex = Bodies.Add(NewBodyHandle);
-
-						const int32 SkeletonBoneIndex = MeshToSkeletonBoneIndex[InsertBone];
-						SkeletonBoneIndexToBodyIndex[SkeletonBoneIndex] = BodyIndex;
-						BodyAnimData[BodyIndex].bIsSimulated = !bKinematic;
-						NamesToHandles.Add(BodySetup->BoneName, NewBodyHandle);
-						BodyIndexToActorHandle[BodyInstance->InstanceBodyIndex] = NewBodyHandle;
-
-						if (BodySetup->CollisionReponse == EBodyCollisionResponse::BodyCollision_Disabled)
-						{
-							IgnoreCollisionActors.Add(NewBodyHandle);
-						}
+						IgnoreCollisionActors.Add(NewBodyHandle);
 					}
 				}
 			}
-		};
-
-		//Insert simulated bodies first to avoid any re-ordering
-		InsertBodiesHelper(/*bSimulated=*/true);
-		InsertBodiesHelper(/*bSimulated=*/false);
+		}
 
 		//Insert joints so that they coincide body order. That is, if we stop simulating all bodies past some index, we can simply ignore joints past a corresponding index without any re-order
 		//For this to work we consider the most last inserted bone in each joint
@@ -730,7 +707,7 @@ void FAnimNode_RigidBody::InitPhysics(const UAnimInstance* InAnimInstance)
 			return false;
 		});
 
-#if WITH_PHYSX
+
 		if(NamesToHandles.Num() > 0)
 		{
 			//constraints
@@ -744,12 +721,7 @@ void FAnimNode_RigidBody::InitPhysics(const UAnimInstance* InAnimInstance)
 				{
 					if (Body1Handle->IsSimulated() || Body2Handle->IsSimulated())
 					{
-						FPhysicsConstraintHandle& ConstraintRef = CI->ConstraintHandle;
-
-#if WITH_CHAOS || WITH_IMMEDIATE_PHYSX || PHYSICS_INTERFACE_LLIMMEDIATE
-                        ensure(false);
-#else
-						PhysicsSimulation->CreateJoint(ConstraintRef.ConstraintData, Body1Handle, Body2Handle);
+						PhysicsSimulation->CreateJoint(CI, Body1Handle, Body2Handle);
 						if (bForceDisableCollisionBetweenConstraintBodies)
 						{
 							int32 BodyIndex1 = UsePhysicsAsset->FindBodyIndex(CI->ConstraintBone1);
@@ -773,7 +745,6 @@ void FAnimNode_RigidBody::InitPhysics(const UAnimInstance* InAnimInstance)
 							FTransform Body2Transform = Body2Handle->GetWorldTransform();
 							BodyAnimData[BodyIndex].RefPoseLength = Body1Transform.GetRelativeTransform(Body2Transform).GetLocation().Size();
 						}
-#endif
 					}
 				}
 
@@ -810,11 +781,10 @@ void FAnimNode_RigidBody::InitPhysics(const UAnimInstance* InAnimInstance)
 
 		PhysicsSimulation->SetIgnoreCollisionPairTable(IgnorePairs);
 		PhysicsSimulation->SetIgnoreCollisionActors(IgnoreCollisionActors);
-#endif
 	}
 }
 
-DECLARE_CYCLE_STAT(TEXT("FAnimNode_Ragdoll::UpdateWorldGeometry"), STAT_ImmediateUpdateWorldGeometry, STATGROUP_ImmediatePhysics);
+DECLARE_CYCLE_STAT(TEXT("FAnimNode_RigidBody::UpdateWorldGeometry"), STAT_ImmediateUpdateWorldGeometry, STATGROUP_ImmediatePhysics);
 
 void FAnimNode_RigidBody::UpdateWorldGeometry(const UWorld& World, const USkeletalMeshComponent& SKC)
 {
@@ -847,7 +817,7 @@ void FAnimNode_RigidBody::UpdateWorldGeometry(const UWorld& World, const USkelet
 	}
 }
 
-DECLARE_CYCLE_STAT(TEXT("FAnimNode_Ragdoll::UpdateWorldForces"), STAT_ImmediateUpdateWorldForces, STATGROUP_ImmediatePhysics);
+DECLARE_CYCLE_STAT(TEXT("FAnimNode_RigidBody::UpdateWorldForces"), STAT_ImmediateUpdateWorldForces, STATGROUP_ImmediatePhysics);
 
 void FAnimNode_RigidBody::UpdateWorldForces(const FTransform& ComponentToWorld, const FTransform& BaseBoneTM)
 {
@@ -864,14 +834,14 @@ void FAnimNode_RigidBody::UpdateWorldForces(const FTransform& ComponentToWorld,
 				if(InvMass > 0.f)
 				{
 					const float StrengthPerBody = PendingRadialForce.bIgnoreMass ? PendingRadialForce.Strength : PendingRadialForce.Strength / (TotalMass * InvMass);
-					ImmediatePhysics::FSimulation::EForceType ForceType;
+					ImmediatePhysics::EForceType ForceType;
 					if (PendingRadialForce.Type == USkeletalMeshComponent::FPendingRadialForces::AddImpulse)
 					{
-						ForceType = PendingRadialForce.bIgnoreMass ? ImmediatePhysics::FSimulation::EForceType::AddVelocity : ImmediatePhysics::FSimulation::EForceType::AddImpulse;
+						ForceType = PendingRadialForce.bIgnoreMass ? ImmediatePhysics::EForceType::AddVelocity : ImmediatePhysics::EForceType::AddImpulse;
 					}
 					else
 					{
-						ForceType = PendingRadialForce.bIgnoreMass ? ImmediatePhysics::FSimulation::EForceType::AddAcceleration : ImmediatePhysics::FSimulation::EForceType::AddForce;
+						ForceType = PendingRadialForce.bIgnoreMass ? ImmediatePhysics::EForceType::AddAcceleration : ImmediatePhysics::EForceType::AddForce;
 					}
 					
 					Body->AddRadialForce(RadialForceOrigin, StrengthPerBody, PendingRadialForce.Radius, PendingRadialForce.Falloff, ForceType);
@@ -976,23 +946,24 @@ void FAnimNode_RigidBody::UpdateInternal(const FAnimationUpdateContext& Context)
 		return;
 	}
 
-	SCOPE_CYCLE_COUNTER(STAT_RigidBody_PreUpdate);
+	SCOPE_CYCLE_COUNTER(STAT_RigidBody_Update);
 
 	// Accumulate deltatime elapsed during update. To be used during evaluation.
 	AccumulatedDeltaTime += Context.AnimInstanceProxy->GetDeltaSeconds();
 
 	if (UnsafeWorld != nullptr)
 	{		
-
-#if WITH_PHYSX && PHYSICS_INTERFACE_PHYSX
 		// Node is valid to evaluate. Simulation is starting.
 		bSimulationStarted = true;
 
 		TArray<FOverlapResult> Overlaps;
 		UnsafeWorld->OverlapMultiByChannel(Overlaps, Bounds.Center, FQuat::Identity, OverlapChannel, FCollisionShape::MakeSphere(Bounds.W), QueryParams, FCollisionResponseParams(ECR_Overlap));
 
+		// @todo(ccaulfield): is there an engine-independent way to do this?
+#if WITH_PHYSX && PHYSICS_INTERFACE_PHYSX
 		SCOPED_SCENE_READ_LOCK(PhysScene ? PhysScene->GetPxScene() : nullptr); //TODO: expose this part to the anim node
-	
+#endif
+
 		for (const FOverlapResult& Overlap : Overlaps)
 		{
 			if (UPrimitiveComponent* OverlapComp = Overlap.GetComponent())
@@ -1000,19 +971,10 @@ void FAnimNode_RigidBody::UpdateInternal(const FAnimationUpdateContext& Context)
 				if (ComponentsInSim.Contains(OverlapComp) == false)
 				{
 					ComponentsInSim.Add(OverlapComp);
-#if WITH_CHAOS || WITH_IMMEDIATE_PHYSX
-                    check(false);
-#else
-					if (PxRigidActor* RigidActor = FPhysicsInterface_PhysX::GetPxRigidActor_AssumesLocked(OverlapComp->BodyInstance.ActorHandle))
-					{
-						PhysicsSimulation->CreateStaticActor(RigidActor, P2UTransform(RigidActor->getGlobalPose()));
-					}
-#endif
+					PhysicsSimulation->CreateActor(ImmediatePhysics::EActorType::StaticActor, &OverlapComp->BodyInstance, OverlapComp->BodyInstance.GetUnrealWorldTransform());
 				}
 			}
 		}
-#endif
-
 		UnsafeWorld = nullptr;
 		PhysScene = nullptr;
 	}
@@ -1080,7 +1042,7 @@ void FAnimNode_RigidBody::InitializeBoneReferences(const FBoneContainer& Require
 		}
 	}
 
-	// New bodies protentially introduced with new LOD
+	// New bodies potentially introduced with new LOD
 	// We'll have to initialize their transform.
 	bCheckForBodyTransformInit = true;
 
diff --git a/Engine/Source/Runtime/AnimGraphRuntime/Private/BoneControllers/AnimNode_RigidBody_Chaos.cpp b/Engine/Source/Runtime/AnimGraphRuntime/Private/BoneControllers/AnimNode_RigidBody_Chaos.cpp
--- a/Engine/Source/Runtime/AnimGraphRuntime/Private/BoneControllers/AnimNode_RigidBody_Chaos.cpp
+++ b/Engine/Source/Runtime/AnimGraphRuntime/Private/BoneControllers/AnimNode_RigidBody_Chaos.cpp
@@ -1,64 +1,45 @@
 // Copyright 1998-2019 Epic Games, Inc. All Rights Reserved.
 
 #include "BoneControllers/AnimNode_RigidBody_Chaos.h"
+
 #include "AnimationRuntime.h"
 #include "Animation/AnimInstanceProxy.h"
 #include "PhysicsEngine/BodySetup.h"
 #include "PhysicsEngine/PhysicsAsset.h"
 #include "PhysicsEngine/PhysicsConstraintTemplate.h"
 #include "GameFramework/PawnMovementComponent.h"
-
-#include "Chaos/DebugDrawQueue.h"
-#include "Chaos/ErrorReporter.h"
-#include "ChaosSolversModule.h"
-#include "ChaosStats.h"
-#include "PBDRigidsSolver.h"
-#include "SolverObjects/SkeletalMeshPhysicsObject.h"
-
+#include "Physics/PhysicsInterfaceCore.h"
+#include "Physics/ImmediatePhysics/ImmediatePhysicsActorHandle.h"
+#include "Physics/ImmediatePhysics/ImmediatePhysicsSimulation.h"
+#include "Physics/ImmediatePhysics/ImmediatePhysicsStats.h"
+#include "PhysicsEngine/PhysicsSettings.h"
 #include "Logging/MessageLog.h"
 
 //#pragma optimize("", off)
 
-#define CHAOS_RBAN_TODO 0
-
 /////////////////////////////////////////////////////
-// FAnimNode_Ragdoll_Chaos
+// FAnimNode_RigidBody_Chaos
 
 #define LOCTEXT_NAMESPACE "ImmediatePhysics"
 
+#if INCLUDE_CHAOS
+
 TAutoConsoleVariable<int32> CVarEnableChaosRigidBodyNode(TEXT("p.ChaosRigidBodyNode"), 1, TEXT("Enables/disables chaos rigid body node updates and evaluations"), ECVF_Scalability);
+TAutoConsoleVariable<int32> CVarChaosRigidBodyLODThreshold(TEXT("p.ChaosRigidBodyLODThreshold"), -1, TEXT("Max LOD that chaos rigid body node is allowed to run on. Provides a global threshold that overrides per-node the LODThreshold property. -1 means no override."), ECVF_Scalability);
 
 FAnimNode_RigidBody_Chaos::FAnimNode_RigidBody_Chaos()
-	: PhysicalMaterial(nullptr)
-
-	, bSimulating(true)
-	, NumIterations(1)
-	, bNotifyCollisions(false)
-	, ObjectType(EObjectStateTypeEnum::Chaos_Object_Kinematic)
-
-	, Density(2.4)  // dense brick
-	, MinMass(0.001)
-	, MaxMass(1.e6)
-
-	, CollisionType(ECollisionTypeEnum::Chaos_Volumetric)
-	, ImplicitShapeParticlesPerUnitArea(0.1)
-	, ImplicitShapeMinNumParticles(0)
-	, ImplicitShapeMaxNumParticles(50)
-	, MinLevelSetResolution(5)
-	, MaxLevelSetResolution(10)
-	, CollisionGroup(0)
-
-	, InitialVelocityType(EInitialVelocityTypeEnum::Chaos_Initial_Velocity_User_Defined)
-	, InitialLinearVelocity(0.f)
-	, InitialAngularVelocity(0.f)
-
+	: QueryParams(NAME_None, FCollisionQueryParams::GetUnknownStatId())
 {
 	ResetSimulatedTeleportType = ETeleportType::None;
+	PhysicsSimulation = nullptr;
 	OverridePhysicsAsset = nullptr;
 	bOverrideWorldGravity = false;
 	CachedBoundsScale = 1.2f;
 	SimulationSpace = ESimulationSpace::ComponentSpace;
 	ExternalForce = FVector::ZeroVector;
+#if WITH_EDITORONLY_DATA
+	bComponentSpaceSimulation_DEPRECATED = true;
+#endif
 	OverrideWorldGravity = FVector::ZeroVector;
 	TotalMass = 0.f;
 	CachedBounds.W = 0;
@@ -82,6 +63,7 @@ FAnimNode_RigidBody_Chaos::FAnimNode_RigidBody_Chaos()
 
 FAnimNode_RigidBody_Chaos::~FAnimNode_RigidBody_Chaos()
 {
+	delete PhysicsSimulation;
 }
 
 void FAnimNode_RigidBody_Chaos::GatherDebugData(FNodeDebugData& DebugData)
@@ -197,7 +179,6 @@ void FAnimNode_RigidBody_Chaos::EvaluateComponentPose_AnyThread(FComponentSpaceP
 
 void FAnimNode_RigidBody_Chaos::InitializeNewBodyTransformsDuringSimulation(FComponentSpacePoseContext& Output, const FTransform& ComponentTransform, const FTransform& BaseBoneTM)
 {
-#if CHAOS_RBAN_TODO
 	for (const FOutputBoneData& OutputData : OutputBoneData)
 	{
 		const int32 BodyIndex = OutputData.BodyIndex;
@@ -220,6 +201,7 @@ void FAnimNode_RigidBody_Chaos::InitializeNewBodyTransformsDuringSimulation(FCom
 
 				const FTransform WSBodyTM = BodyRelativeTransform * Bodies[OutputData.ParentBodyIndex]->GetWorldTransform();
 				Bodies[BodyIndex]->SetWorldTransform(WSBodyTM);
+				BodyAnimData[BodyIndex].RefPoseLength = BodyRelativeTransform.GetLocation().Size();
 			}
 			// If we don't have a parent body, then we can just grab the incoming pose in component space.
 			else
@@ -231,51 +213,32 @@ void FAnimNode_RigidBody_Chaos::InitializeNewBodyTransformsDuringSimulation(FCom
 			}
 		}
 	}
-#endif
 }
 
-DECLARE_CYCLE_STAT(TEXT("RigidBody_Chaos_Eval"), STAT_RigidBody_Chaos_Eval, STATGROUP_Anim);
-DECLARE_CYCLE_STAT(TEXT("FAnimNode_RigidBody_Chaos::EvaluateSkeletalControl_AnyThread"), STAT_ImmediateEvaluateSkeletalControl_Chaos, STATGROUP_Chaos);
+DECLARE_CYCLE_STAT(TEXT("RigidBody_Eval_Chaos"), STAT_RigidBody_Chaos_Eval, STATGROUP_Anim);
+DECLARE_CYCLE_STAT(TEXT("FAnimNode_RigidBody::EvaluateSkeletalControl_AnyThread_Chaos"), STAT_ImmediateEvaluateSkeletalControl_Chaos, STATGROUP_ImmediatePhysics);
 
 void FAnimNode_RigidBody_Chaos::EvaluateSkeletalControl_AnyThread(FComponentSpacePoseContext& Output, TArray<FBoneTransform>& OutBoneTransforms)
 {
-#if INCLUDE_CHAOS
 	SCOPE_CYCLE_COUNTER(STAT_RigidBody_Chaos_Eval);
 	SCOPE_CYCLE_COUNTER(STAT_ImmediateEvaluateSkeletalControl_Chaos);
-	//SCOPED_NAMED_EVENT_TEXT("FAnimNode_Ragdoll::EvaluateSkeletalControl_AnyThread", FColor::Magenta);
+	//SCOPED_NAMED_EVENT_TEXT("FAnimNode_RigidBody::EvaluateSkeletalControl_AnyThread_Chaos", FColor::Magenta);
 
 	// Update our eval counter, and decide whether we need to reset simulated bodies, if our anim instance hasn't updated in a while.
-	if(EvalCounter.HasEverBeenUpdated() && !EvalCounter.WasSynchronizedLastFrame(Output.AnimInstanceProxy->GetEvaluationCounter()))
+	if (EvalCounter.HasEverBeenUpdated())
 	{
-		ResetSimulatedTeleportType = ETeleportType::ResetPhysics;
+		// Always propagate skip rate as it can go up and down between updates
+		EvalCounter.SetMaxSkippedFrames(Output.AnimInstanceProxy->GetEvaluationCounter().GetMaxSkippedFrames());
+		if (!EvalCounter.WasSynchronizedLastFrame(Output.AnimInstanceProxy->GetEvaluationCounter()))
+		{
+			ResetSimulatedTeleportType = ETeleportType::ResetPhysics;
+		}
 	}
 	EvalCounter.SynchronizeWith(Output.AnimInstanceProxy->GetEvaluationCounter());
 
 	const float DeltaSeconds = AccumulatedDeltaTime;
 	AccumulatedDeltaTime = 0.f;
 
-	if (Solver && Solver->Enabled())
-	{
-		PhysicsObject->CaptureInputs(DeltaSeconds,
-			[this, &Output](const float Dt, FSkeletalMeshPhysicsObjectParams & OutPhysicsParams) -> bool
-			{
-				return UpdatePhysicsInputs(Output, Dt, OutPhysicsParams.BoneHierarchy);
-			});
-
-		Solver->AdvanceSolverBy(DeltaSeconds);
-		
-		// @todo(ccaulfield): this feels clumsy
-		PhysicsObject->CacheResults();
-		PhysicsObject->FlipCache();
-		PhysicsObject->SyncToCache();
-
-		if (const FSkeletalMeshPhysicsObjectOutputs* PhysicsOutputs = PhysicsObject->GetOutputs())
-		{
-			UpdateAnimNodeOutputs(PhysicsObject->GetBoneHierarchy(), *PhysicsOutputs, Output, OutBoneTransforms);
-		}
-	}
-
-#if CHAOS_RBAN_TODO
 	if (CVarEnableRigidBodyNode.GetValueOnAnyThread() != 0 && PhysicsSimulation)	
 	{
 		const FBoneContainer& BoneContainer = Output.Pose.GetPose().GetBoneContainer();
@@ -284,10 +247,6 @@ void FAnimNode_RigidBody_Chaos::EvaluateSkeletalControl_AnyThread(FComponentSpac
 		{
 			PreviousCompWorldSpaceTM = CompWorldSpaceTM;
 		}
-
-		// @todo(ccaulfield) fix
-		const FTransform BaseBoneTM = FTransform::Identity;
-
 		const FTransform BaseBoneTM = Output.Pose.GetComponentSpaceTransform(BaseBoneRef.GetCompactPoseIndex(BoneContainer));
 
 		// Initialize potential new bodies because of LOD change.
@@ -299,7 +258,8 @@ void FAnimNode_RigidBody_Chaos::EvaluateSkeletalControl_AnyThread(FComponentSpac
 
 		// If time advances, update simulation
 		// Reset if necessary
-		if (ResetSimulatedTeleportType != ETeleportType::None)
+		bool bDynamicsReset = (ResetSimulatedTeleportType != ETeleportType::None);
+		if (bDynamicsReset)
 		{
 			// Capture bone velocities if we have captured a bone velocity pose.
 			if (bTransferBoneVelocities && (CapturedBoneVelocityPose.GetPose().GetNumBones() > 0))
@@ -369,6 +329,10 @@ void FAnimNode_RigidBody_Chaos::EvaluateSkeletalControl_AnyThread(FComponentSpac
 
 						BodyTM = ConvertCSTransformToSimSpaceChaos(SimulationSpace, ComponentSpaceTM, CompWorldSpaceTM, BaseBoneTM);
 						Bodies[BodyIndex]->SetWorldTransform(BodyTM);
+						if (OutputData.ParentBodyIndex != INDEX_NONE)
+						{
+							BodyAnimData[BodyIndex].RefPoseLength = BodyTM.GetRelativeTransform(Bodies[OutputData.ParentBodyIndex]->GetWorldTransform()).GetLocation().Size();
+						}
 					}
 				}
 				break;
@@ -384,6 +348,10 @@ void FAnimNode_RigidBody_Chaos::EvaluateSkeletalControl_AnyThread(FComponentSpac
 						const FTransform& ComponentSpaceTM = Output.Pose.GetComponentSpaceTransform(OutputData.CompactPoseBoneIndex);
 						const FTransform BodyTM = ConvertCSTransformToSimSpaceChaos(SimulationSpace, ComponentSpaceTM, CompWorldSpaceTM, BaseBoneTM);
 						Bodies[BodyIndex]->SetWorldTransform(BodyTM);
+						if (OutputData.ParentBodyIndex != INDEX_NONE)
+						{
+							BodyAnimData[BodyIndex].RefPoseLength = BodyTM.GetRelativeTransform(Bodies[OutputData.ParentBodyIndex]->GetWorldTransform()).GetLocation().Size();
+						}
 					}
 				}
 				break;
@@ -395,9 +363,9 @@ void FAnimNode_RigidBody_Chaos::EvaluateSkeletalControl_AnyThread(FComponentSpac
 			PreviousComponentLinearVelocity = FVector::ZeroVector;
 		}
 		// Only need to tick physics if we didn't reset and we have some time to simulate
-		else if(DeltaSeconds > 0.0f)
+		if((bSimulateAnimPhysicsAfterReset || !bDynamicsReset) && DeltaSeconds > AnimPhysicsMinDeltaTime)
 		{
-		// Transfer bone velocities previously captured.
+			// Transfer bone velocities previously captured.
 			if (bTransferBoneVelocities && (CapturedBoneVelocityPose.GetPose().GetNumBones() > 0))
 			{
 				for (const FOutputBoneData& OutputData : OutputBoneData)
@@ -407,13 +375,14 @@ void FAnimNode_RigidBody_Chaos::EvaluateSkeletalControl_AnyThread(FComponentSpac
 
 					if (BodyData.bIsSimulated)
 					{
-						ImmediatePhysics::FActorHandle* Body = Bodies[BodyIndex];
+						ImmediatePhysics_Chaos::FActorHandle* Body = Bodies[BodyIndex];
 						Body->SetLinearVelocity(BodyData.TransferedBoneLinearVelocity);
 
 						const FQuat AngularVelocity = BodyData.TransferedBoneAngularVelocity;
 						Body->SetAngularVelocity(AngularVelocity.GetRotationAxis() * AngularVelocity.GetAngle());
 					}
 				}
+
 				// Free up our captured pose after it's been used.
 				CapturedBoneVelocityPose.Empty();
 			}
@@ -439,7 +408,7 @@ void FAnimNode_RigidBody_Chaos::EvaluateSkeletalControl_AnyThread(FComponentSpac
 
 					if (BodyData.bIsSimulated)
 					{
-						ImmediatePhysics::FActorHandle* Body = Bodies[BodyIndex];
+						ImmediatePhysics_Chaos::FActorHandle* Body = Bodies[BodyIndex];
 
 						// Apply 
 						const float BodyInvMass = Body->GetInverseMass();
@@ -489,7 +458,7 @@ void FAnimNode_RigidBody_Chaos::EvaluateSkeletalControl_AnyThread(FComponentSpac
 				PhysicsSimulation->Simulate_AssumesLocked(StepDeltaTime, SimSpaceGravity);
 			}
 		}
-
+		
 		//write back to animation system
 		for (const FOutputBoneData& OutputData : OutputBoneData)
 		{
@@ -555,8 +524,6 @@ void FAnimNode_RigidBody_Chaos::EvaluateSkeletalControl_AnyThread(FComponentSpac
 
 		PreviousCompWorldSpaceTM = CompWorldSpaceTM;
 	}
-#endif // CHAOS_RBAN_TODO
-#endif // INCLUDE_CHAOS
 }
 
 void ComputeBodyInsertionOrderChaos(TArray<FBoneIndexType>& InsertionOrder, const USkeletalMeshComponent& SKC)
@@ -604,92 +571,6 @@ void ComputeBodyInsertionOrderChaos(TArray<FBoneIndexType>& InsertionOrder, cons
 	}
 }
 
-#if INCLUDE_CHAOS
-
-void FAnimNode_RigidBody_Chaos::PhysicsObjectInitCallback(const USkeletalMeshComponent* InSkelMeshComponent, const UAnimInstance* InAnimInstance, FSkeletalMeshPhysicsObjectParams& OutPhysicsParams)
-{
-	OutPhysicsParams.bSimulating = bSimulating;
-
-	// @todo(ccaulfield): do we need this?
-	//GetPathName(this, OutParams.Name);
-	if (InitialVelocityType == EInitialVelocityTypeEnum::Chaos_Initial_Velocity_User_Defined)
-	{
-		OutPhysicsParams.InitialLinearVelocity = InitialLinearVelocity;
-		OutPhysicsParams.InitialAngularVelocity = InitialAngularVelocity;
-	}
-
-	//OutParams.PhysicalMaterial = MakeSerializable(PhysicsMaterial);
-	OutPhysicsParams.ObjectType = ObjectType;
-
-	OutPhysicsParams.Density = Density;
-	OutPhysicsParams.MinMass = MinMass;
-	OutPhysicsParams.MaxMass = MaxMass;
-
-	OutPhysicsParams.CollisionType = CollisionType;
-	OutPhysicsParams.ParticlesPerUnitArea = ImplicitShapeParticlesPerUnitArea;
-	OutPhysicsParams.MinNumParticles = ImplicitShapeMinNumParticles;
-	OutPhysicsParams.MaxNumParticles = ImplicitShapeMaxNumParticles;
-	OutPhysicsParams.MinRes = MinLevelSetResolution;
-	OutPhysicsParams.MaxRes = MaxLevelSetResolution;
-	OutPhysicsParams.CollisionGroup = CollisionGroup;
-
-	const USkeletalMesh* SkeletalMesh = InSkelMeshComponent->SkeletalMesh;
-	const AActor* OwningActor = InSkelMeshComponent->GetOwner();
-	if (SkeletalMesh && OwningActor)
-	{
-		const UPhysicsAsset* PhysicsAsset = OverridePhysicsAsset ? OverridePhysicsAsset : SkeletalMesh->PhysicsAsset;
-		FPhysicsAssetSimulationUtil::BuildParams(OwningActor, OwningActor, InSkelMeshComponent, PhysicsAsset, OutPhysicsParams);
-	}
-}
-
-bool FAnimNode_RigidBody_Chaos::UpdatePhysicsInputs(FComponentSpacePoseContext& InPoseContext, const float Dt, FBoneHierarchy& InOutBoneHierarchy)
-{
-	InOutBoneHierarchy.PrepareForUpdate();
-	for (int32 BoneIndex : InOutBoneHierarchy.GetBoneIndices())
-	{
-		InOutBoneHierarchy.SetAnimLocalSpaceTransform(BoneIndex, InPoseContext.Pose.GetLocalSpaceTransform(FCompactPoseBoneIndex(BoneIndex)));
-	}
-	InOutBoneHierarchy.SetActorWorldSpaceTransform(InPoseContext.AnimInstanceProxy->GetComponentTransform());
-	InOutBoneHierarchy.PrepareAnimWorldSpaceTransforms();
-
-	return true;
-}
-
-void FAnimNode_RigidBody_Chaos::UpdateAnimNodeOutputs(const FBoneHierarchy& InBoneHierarchy, const FSkeletalMeshPhysicsObjectOutputs& InPhysicsOutputs, FComponentSpacePoseContext& PoseContext, TArray<FBoneTransform>& OutComponentSpaceBoneTransforms)
-{
-	const FTransform& ComponentWorldSpaceTransform = PoseContext.AnimInstanceProxy->GetComponentTransform();
-
-	for (const int32 BoneIndex : InBoneHierarchy.GetBoneIndices())
-	{
-		const FAnalyticImplicitGroup* ShapeGroup = InBoneHierarchy.GetAnalyticShapeGroup(BoneIndex);
-
-		// The first time we have no body, we are done
-		// @todo(ccaulfield): is this right?
-		if ((ShapeGroup == nullptr) || (ShapeGroup->GetRigidBodyId() == INDEX_NONE))
-		{
-			break;
-		}
-
-		// Shouldn't have to process Kinematic bodies
-		// @todo(ccaulfield): should store actual state of particle
-		if (ShapeGroup->GetRigidBodyState() == EObjectStateTypeEnum::Chaos_Object_Kinematic)
-		{
-			continue;
-		}
-
-		// @todo(ccaulfield): This should be pulling from an updated Hierarchy (so we support non-simulated child bones). See FSkeletalMeshPhysicsObject::SyncToCache
-		const int32 TransformIndex = InBoneHierarchy.GetTransformIndex(BoneIndex);
-		if (TransformIndex != INDEX_NONE)
-		{
-			FTransform ComponentSpaceBoneTransform = InPhysicsOutputs.Transforms[TransformIndex].GetRelativeTransform(ComponentWorldSpaceTransform);
-			OutComponentSpaceBoneTransforms.Emplace(FBoneTransform(FCompactPoseBoneIndex(BoneIndex), ComponentSpaceBoneTransform));
-		}
-	}
-}
-
-
-#endif // INCLUDE_CHAOS
-
 void FAnimNode_RigidBody_Chaos::InitPhysics(const UAnimInstance* InAnimInstance)
 {
 	const USkeletalMeshComponent* SkeletalMeshComp = InAnimInstance->GetSkelMeshComponent();
@@ -706,60 +587,48 @@ void FAnimNode_RigidBody_Chaos::InitPhysics(const UAnimInstance* InAnimInstance)
 	const FSkeletonToMeshLinkup& SkeletonToMeshLinkupTable = SkeletonAsset->LinkupCache[SkelMeshLinkupIndex];
 	const TArray<int32>& MeshToSkeletonBoneIndex = SkeletonToMeshLinkupTable.MeshToSkeletonTable;
 	
-	// @todo(ccaulfield) fix
-	//const int32 NumSkeletonBones = SkeletonAsset->GetReferenceSkeleton().GetNum();
-	//SkeletonBoneIndexToBodyIndex.Reset(NumSkeletonBones);
-	//SkeletonBoneIndexToBodyIndex.Init(INDEX_NONE, NumSkeletonBones);
+	const int32 NumSkeletonBones = SkeletonAsset->GetReferenceSkeleton().GetNum();
+	SkeletonBoneIndexToBodyIndex.Reset(NumSkeletonBones);
+	SkeletonBoneIndexToBodyIndex.Init(INDEX_NONE, NumSkeletonBones);
 
-	PreviousTransform = InAnimInstance->GetSkelMeshComponent()->GetComponentToWorld();
+	PreviousTransform = SkeletalMeshComp->GetComponentToWorld();
 
-	if(UsePhysicsAsset)
+	if (UPhysicsSettings* Settings = UPhysicsSettings::Get())
 	{
-#if INCLUDE_CHAOS
-		Solver = Chaos::FPBDRigidsSolver::FAccessor::CreateSolver();
-		Solver->SetIterations(NumIterations);
-
-		PhysicsObject = new FSkeletalMeshPhysicsObject(
-			InAnimInstance->GetOwningActor(),
-			[this, InAnimInstance, SkeletalMeshComp](FSkeletalMeshPhysicsObjectParams & OutPhysicsParams)
-			{
-				PhysicsObjectInitCallback(SkeletalMeshComp, InAnimInstance, OutPhysicsParams);
-			});
-
-		//PhysicsMaterial = MakeUnique<Chaos::TChaosPhysicsMaterial<float>>();
-
-		// @todo(ccaulfield): should only be one call for all this...?
-		// @todo(ccaulfield): remove the floor
-		Solver->RegisterObject(PhysicsObject);
-		PhysicsObject->Initialize();
-		Solver->SetHasFloor(true);
-		Solver->SetIsFloorAnalytic(true);
-		Solver->SetEnabled(true);
-#endif
+		AnimPhysicsMinDeltaTime = Settings->AnimPhysicsMinDeltaTime;
+		bSimulateAnimPhysicsAfterReset = Settings->bSimulateAnimPhysicsAfterReset;
+	}
+	else
+	{
+		AnimPhysicsMinDeltaTime = 0.f;
+		bSimulateAnimPhysicsAfterReset = false;
+	}
 
-#if CHAOS_RBAN_TODO
+	if(UsePhysicsAsset)
+	{
 		delete PhysicsSimulation;
-		PhysicsSimulation = new ImmediatePhysics::FSimulation();
+		PhysicsSimulation = new ImmediatePhysics_Chaos::FSimulation();
 		const int32 NumBodies = UsePhysicsAsset->SkeletalBodySetups.Num();
 		Bodies.Empty(NumBodies);
 		ComponentsInSim.Reset();
 		BodyAnimData.Reset(NumBodies);
 		BodyAnimData.AddDefaulted(NumBodies);
 		TotalMass = 0.f;
-		
+
+		// Instantiate a FBodyInstance/FConstraintInstance set that will be cloned into the Immediate Physics sim.
 		TArray<FBodyInstance*> HighLevelBodyInstances;
 		TArray<FConstraintInstance*> HighLevelConstraintInstances;
 		SkeletalMeshComp->InstantiatePhysicsAssetRefPose(*UsePhysicsAsset, SimulationSpace == ESimulationSpace::WorldSpace ? SkeletalMeshComp->GetComponentToWorld().GetScale3D() : FVector(1.f), HighLevelBodyInstances, HighLevelConstraintInstances);
 
-		TMap<FName, ImmediatePhysics::FActorHandle*> NamesToHandles;
-		TArray<ImmediatePhysics::FActorHandle*> IgnoreCollisionActors;
+		TMap<FName, ImmediatePhysics_Chaos::FActorHandle*> NamesToHandles;
+		TArray<ImmediatePhysics_Chaos::FActorHandle*> IgnoreCollisionActors;
 
 		TArray<FBoneIndexType> InsertionOrder;
 		ComputeBodyInsertionOrderChaos(InsertionOrder, *SkeletalMeshComp);
 
 		const int32 NumBonesLOD0 = InsertionOrder.Num();
 
-		TArray<ImmediatePhysics::FActorHandle*> BodyIndexToActorHandle;
+		TArray<ImmediatePhysics_Chaos::FActorHandle*> BodyIndexToActorHandle;
 		BodyIndexToActorHandle.AddZeroed(NumBonesLOD0);
 
 		TArray<FBodyInstance*> BodiesSorted;
@@ -773,63 +642,37 @@ void FAnimNode_RigidBody_Chaos::InitPhysics(const UAnimInstance* InAnimInstance)
 			}
 		}
 
-		auto InsertBodiesHelper = [&](bool bSimulatedBodies)
+		// Create the immediate physics bodies
+		for (FBoneIndexType InsertBone : InsertionOrder)
 		{
-			for (FBoneIndexType InsertBone : InsertionOrder)
+			if (FBodyInstance* BodyInstance = BodiesSorted[InsertBone])
 			{
-				if (FBodyInstance* BodyInstance = BodiesSorted[InsertBone])
-				{
-					UBodySetup* BodySetup = UsePhysicsAsset->SkeletalBodySetups[BodyInstance->InstanceBodyIndex];
-					const bool bKinematic = BodySetup->PhysicsType != EPhysicsType::PhysType_Simulated;
-					const FTransform& LastTransform = SkeletalMeshComp->GetBoneTransform(InsertBone);	//This is out of date, but will still give our bodies an initial setup that matches the constraints (TODO: use refpose)
+				UBodySetup* BodySetup = UsePhysicsAsset->SkeletalBodySetups[BodyInstance->InstanceBodyIndex];
 
-					ImmediatePhysics::FActorHandle* NewBodyHandle = nullptr;
-					if (bSimulatedBodies && !bKinematic)
+				bool bSimulated = (BodySetup->PhysicsType == EPhysicsType::PhysType_Simulated);
+				ImmediatePhysics_Chaos::EActorType ActorType = bSimulated ?  ImmediatePhysics_Chaos::EActorType::DynamicActor : ImmediatePhysics_Chaos::EActorType::KinematicActor;
+				ImmediatePhysics_Chaos::FActorHandle* NewBodyHandle = PhysicsSimulation->CreateActor(ActorType, BodyInstance, BodyInstance->GetUnrealWorldTransform());
+				if (NewBodyHandle)
+				{
+					if (bSimulated)
 					{
-#if WITH_PHYSX && PHYSICS_INTERFACE_PHYSX
-						PxRigidActor* PRigidActor = FPhysicsInterface::GetPxRigidActor_AssumesLocked(BodyInstance->GetPhysicsActorHandle());
-						if(PxRigidDynamic* PDynamic = PRigidActor->is<PxRigidDynamic>())
-						{
-							NewBodyHandle = PhysicsSimulation->CreateDynamicActor(PDynamic, LastTransform);
-						checkSlow(NewBodyHandle);
 						const float InvMass = NewBodyHandle->GetInverseMass();
 						TotalMass += InvMass > 0.f ? 1.f / InvMass : 0.f;
-						}
-#endif
 					}
-					else if (bSimulatedBodies && bKinematic)
+					const int32 BodyIndex = Bodies.Add(NewBodyHandle);
+					const int32 SkeletonBoneIndex = MeshToSkeletonBoneIndex[InsertBone];
+					SkeletonBoneIndexToBodyIndex[SkeletonBoneIndex] = BodyIndex;
+					BodyAnimData[BodyIndex].bIsSimulated = bSimulated;
+					NamesToHandles.Add(BodySetup->BoneName, NewBodyHandle);
+					BodyIndexToActorHandle[BodyInstance->InstanceBodyIndex] = NewBodyHandle;
+
+					if (BodySetup->CollisionReponse == EBodyCollisionResponse::BodyCollision_Disabled)
 					{
-#if WITH_PHYSX && PHYSICS_INTERFACE_PHYSX
-						PxRigidActor* PRigidActor = FPhysicsInterface::GetPxRigidActor_AssumesLocked(BodyInstance->GetPhysicsActorHandle());
-						if(PxRigidBody* PRigidBody = PRigidActor->is<PxRigidBody>())
-						{
-							NewBodyHandle = PhysicsSimulation->CreateKinematicActor(PRigidBody, LastTransform);
-						}
-#endif
-					}
-
-					if (NewBodyHandle)
-					{
-						const int32 BodyIndex = Bodies.Add(NewBodyHandle);
-
-						const int32 SkeletonBoneIndex = MeshToSkeletonBoneIndex[InsertBone];
-						SkeletonBoneIndexToBodyIndex[SkeletonBoneIndex] = BodyIndex;
-						BodyAnimData[BodyIndex].bIsSimulated = !bKinematic;
-						NamesToHandles.Add(BodySetup->BoneName, NewBodyHandle);
-						BodyIndexToActorHandle[BodyInstance->InstanceBodyIndex] = NewBodyHandle;
-
-						if (BodySetup->CollisionReponse == EBodyCollisionResponse::BodyCollision_Disabled)
-						{
-							IgnoreCollisionActors.Add(NewBodyHandle);
-						}
+						IgnoreCollisionActors.Add(NewBodyHandle);
 					}
 				}
 			}
-		};
-
-		//Insert simulated bodies first to avoid any re-ordering
-		InsertBodiesHelper(/*bSimulated=*/true);
-		InsertBodiesHelper(/*bSimulated=*/false);
+		}
 
 		//Insert joints so that they coincide body order. That is, if we stop simulating all bodies past some index, we can simply ignore joints past a corresponding index without any re-order
 		//For this to work we consider the most last inserted bone in each joint
@@ -860,26 +703,21 @@ void FAnimNode_RigidBody_Chaos::InitPhysics(const UAnimInstance* InAnimInstance)
 			return false;
 		});
 
-#if WITH_PHYSX
+
 		if(NamesToHandles.Num() > 0)
 		{
 			//constraints
 			for(int32 ConstraintIdx = 0; ConstraintIdx < HighLevelConstraintInstances.Num(); ++ConstraintIdx)
 			{
 				FConstraintInstance* CI = HighLevelConstraintInstances[ConstraintIdx];
-				ImmediatePhysics::FActorHandle* Body1Handle = NamesToHandles.FindRef(CI->ConstraintBone1);
-				ImmediatePhysics::FActorHandle* Body2Handle = NamesToHandles.FindRef(CI->ConstraintBone2);
+				ImmediatePhysics_Chaos::FActorHandle* Body1Handle = NamesToHandles.FindRef(CI->ConstraintBone1);
+				ImmediatePhysics_Chaos::FActorHandle* Body2Handle = NamesToHandles.FindRef(CI->ConstraintBone2);
 
 				if(Body1Handle && Body2Handle)
 				{
 					if (Body1Handle->IsSimulated() || Body2Handle->IsSimulated())
 					{
-						FPhysicsConstraintHandle& ConstraintRef = CI->ConstraintHandle;
-
-#if WITH_CHAOS || WITH_IMMEDIATE_PHYSX || PHYSICS_INTERFACE_LLIMMEDIATE
-						ensure(false);
-#else
-						PhysicsSimulation->CreateJoint(ConstraintRef.ConstraintData, Body1Handle, Body2Handle);
+						PhysicsSimulation->CreateJoint(CI, Body1Handle, Body2Handle);
 						if (bForceDisableCollisionBetweenConstraintBodies)
 						{
 							int32 BodyIndex1 = UsePhysicsAsset->FindBodyIndex(CI->ConstraintBone1);
@@ -903,7 +741,6 @@ void FAnimNode_RigidBody_Chaos::InitPhysics(const UAnimInstance* InAnimInstance)
 							FTransform Body2Transform = Body2Handle->GetWorldTransform();
 							BodyAnimData[BodyIndex].RefPoseLength = Body1Transform.GetRelativeTransform(Body2Transform).GetLocation().Size();
 						}
-#endif
 					}
 				}
 
@@ -928,11 +765,11 @@ void FAnimNode_RigidBody_Chaos::InitPhysics(const UAnimInstance* InAnimInstance)
 		HighLevelBodyInstances.Empty();
 		BodiesSorted.Empty();
 
-		TArray<ImmediatePhysics::FSimulation::FIgnorePair> IgnorePairs;
+		TArray<ImmediatePhysics_Chaos::FSimulation::FIgnorePair> IgnorePairs;
 		const TMap<FRigidBodyIndexPair, bool>& DisableTable = UsePhysicsAsset->CollisionDisableTable;
 		for(auto ConstItr = DisableTable.CreateConstIterator(); ConstItr; ++ConstItr)
 		{
-			ImmediatePhysics::FSimulation::FIgnorePair Pair;
+			ImmediatePhysics_Chaos::FSimulation::FIgnorePair Pair;
 			Pair.A = BodyIndexToActorHandle[ConstItr.Key().Indices[0]];
 			Pair.B = BodyIndexToActorHandle[ConstItr.Key().Indices[1]];
 			IgnorePairs.Add(Pair);
@@ -940,16 +777,13 @@ void FAnimNode_RigidBody_Chaos::InitPhysics(const UAnimInstance* InAnimInstance)
 
 		PhysicsSimulation->SetIgnoreCollisionPairTable(IgnorePairs);
 		PhysicsSimulation->SetIgnoreCollisionActors(IgnoreCollisionActors);
-#endif
-#endif // TODO
 	}
 }
 
-DECLARE_CYCLE_STAT(TEXT("FAnimNode_RigidBody_Chaos::UpdateWorldGeometry"), STAT_ImmediateUpdateWorldGeometry_Chaos, STATGROUP_Chaos);
+DECLARE_CYCLE_STAT(TEXT("FAnimNode_RigidBody::UpdateWorldGeometry_Chaos"), STAT_ImmediateUpdateWorldGeometry_Chaos, STATGROUP_ImmediatePhysics);
 
 void FAnimNode_RigidBody_Chaos::UpdateWorldGeometry(const UWorld& World, const USkeletalMeshComponent& SKC)
 {
-#if CHAOS_RBAN_TODO
 	SCOPE_CYCLE_COUNTER(STAT_ImmediateUpdateWorldGeometry_Chaos);
 	QueryParams = FCollisionQueryParams(SCENE_QUERY_STAT(RagdollNodeFindGeometry), /*bTraceComplex=*/false);
 #if WITH_EDITOR
@@ -977,14 +811,12 @@ void FAnimNode_RigidBody_Chaos::UpdateWorldGeometry(const UWorld& World, const U
 		PhysScene = World.GetPhysicsScene();
 		UnsafeWorld = &World;
 	}
-#endif
 }
 
-DECLARE_CYCLE_STAT(TEXT("FAnimNode_RigidBody_Chaos::UpdateWorldForces"), STAT_ImmediateUpdateWorldForces_Chaos, STATGROUP_Chaos);
+DECLARE_CYCLE_STAT(TEXT("FAnimNode_RigidBody::UpdateWorldForces_Chaos"), STAT_ImmediateUpdateWorldForces_Chaos, STATGROUP_ImmediatePhysics);
 
 void FAnimNode_RigidBody_Chaos::UpdateWorldForces(const FTransform& ComponentToWorld, const FTransform& BaseBoneTM)
 {
-#if CHAOS_RBAN_TODO
 	SCOPE_CYCLE_COUNTER(STAT_ImmediateUpdateWorldForces_Chaos);
 
 	if(TotalMass > 0.f)
@@ -992,20 +824,20 @@ void FAnimNode_RigidBody_Chaos::UpdateWorldForces(const FTransform& ComponentToW
 		for (const USkeletalMeshComponent::FPendingRadialForces& PendingRadialForce : PendingRadialForces)
 		{
 			const FVector RadialForceOrigin = WorldPositionToSpaceChaos(SimulationSpace, PendingRadialForce.Origin, ComponentToWorld, BaseBoneTM);
-			for(ImmediatePhysics::FActorHandle* Body : Bodies)
+			for(ImmediatePhysics_Chaos::FActorHandle* Body : Bodies)
 			{
 				const float InvMass = Body->GetInverseMass();
 				if(InvMass > 0.f)
 				{
 					const float StrengthPerBody = PendingRadialForce.bIgnoreMass ? PendingRadialForce.Strength : PendingRadialForce.Strength / (TotalMass * InvMass);
-					ImmediatePhysics::FSimulation::EForceType ForceType;
+					ImmediatePhysics_Chaos::EForceType ForceType;
 					if (PendingRadialForce.Type == USkeletalMeshComponent::FPendingRadialForces::AddImpulse)
 					{
-						ForceType = PendingRadialForce.bIgnoreMass ? ImmediatePhysics::FSimulation::EForceType::AddVelocity : ImmediatePhysics::FSimulation::EForceType::AddImpulse;
+						ForceType = PendingRadialForce.bIgnoreMass ? ImmediatePhysics_Chaos::EForceType::AddVelocity : ImmediatePhysics_Chaos::EForceType::AddImpulse;
 					}
 					else
 					{
-						ForceType = PendingRadialForce.bIgnoreMass ? ImmediatePhysics::FSimulation::EForceType::AddAcceleration : ImmediatePhysics::FSimulation::EForceType::AddForce;
+						ForceType = PendingRadialForce.bIgnoreMass ? ImmediatePhysics_Chaos::EForceType::AddAcceleration : ImmediatePhysics_Chaos::EForceType::AddForce;
 					}
 					
 					Body->AddRadialForce(RadialForceOrigin, StrengthPerBody, PendingRadialForce.Radius, PendingRadialForce.Falloff, ForceType);
@@ -1016,7 +848,7 @@ void FAnimNode_RigidBody_Chaos::UpdateWorldForces(const FTransform& ComponentToW
 		if(!ExternalForce.IsNearlyZero())
 		{
 			const FVector ExternalForceInSimSpace = WorldVectorToSpaceNoScaleChaos(SimulationSpace, ExternalForce, ComponentToWorld, BaseBoneTM);
-			for (ImmediatePhysics::FActorHandle* Body : Bodies)
+			for (ImmediatePhysics_Chaos::FActorHandle* Body : Bodies)
 			{
 				const float InvMass = Body->GetInverseMass();
 				if (InvMass > 0.f)
@@ -1026,7 +858,6 @@ void FAnimNode_RigidBody_Chaos::UpdateWorldForces(const FTransform& ComponentToW
 			}
 		}
 	}
-#endif
 }
 
 bool FAnimNode_RigidBody_Chaos::NeedsDynamicReset() const
@@ -1041,11 +872,10 @@ void FAnimNode_RigidBody_Chaos::ResetDynamics(ETeleportType InTeleportType)
 	ResetSimulatedTeleportType = ((InTeleportType > ResetSimulatedTeleportType) ? InTeleportType : ResetSimulatedTeleportType);
 }
 
-DECLARE_CYCLE_STAT(TEXT("RigidBody_Chaos_PreUpdate"), STAT_RigidBody_Chaos_PreUpdate, STATGROUP_Anim);
+DECLARE_CYCLE_STAT(TEXT("RigidBody_PreUpdate_Chaos"), STAT_RigidBody_Chaos_PreUpdate, STATGROUP_Anim);
 
 void FAnimNode_RigidBody_Chaos::PreUpdate(const UAnimInstance* InAnimInstance)
 {
-#if CHAOS_RBAN_TODO
 	// Don't update geometry if RBN is disabled
 	if(CVarEnableRigidBodyNode.GetValueOnAnyThread() == 0)
 	{
@@ -1079,8 +909,7 @@ void FAnimNode_RigidBody_Chaos::PreUpdate(const UAnimInstance* InAnimInstance)
 
 		PreviousTransform = CurrentTransform;
 		CurrentTransform = SKC->GetComponentToWorld();
-	}
-#endif
+	}	
 }
 
 int32 FAnimNode_RigidBody_Chaos::GetLODThreshold() const
@@ -1102,7 +931,7 @@ int32 FAnimNode_RigidBody_Chaos::GetLODThreshold() const
 	}
 }
 
-DECLARE_CYCLE_STAT(TEXT("RigidBody_Chaos_Update"), STAT_RigidBody_Chaos_Update, STATGROUP_Anim);
+DECLARE_CYCLE_STAT(TEXT("RigidBody_Update_Chaos"), STAT_RigidBody_Chaos_Update, STATGROUP_Anim);
 
 void FAnimNode_RigidBody_Chaos::UpdateInternal(const FAnimationUpdateContext& Context)
 {
@@ -1112,24 +941,24 @@ void FAnimNode_RigidBody_Chaos::UpdateInternal(const FAnimationUpdateContext& Co
 		return;
 	}
 
-	SCOPE_CYCLE_COUNTER(STAT_RigidBody_Chaos_PreUpdate);
+	SCOPE_CYCLE_COUNTER(STAT_RigidBody_Chaos_Update);
 
 	// Accumulate deltatime elapsed during update. To be used during evaluation.
 	AccumulatedDeltaTime += Context.AnimInstanceProxy->GetDeltaSeconds();
 
-#if CHAOS_RBAN_TODO
 	if (UnsafeWorld != nullptr)
 	{		
-
-#if WITH_PHYSX && PHYSICS_INTERFACE_PHYSX
 		// Node is valid to evaluate. Simulation is starting.
 		bSimulationStarted = true;
 
 		TArray<FOverlapResult> Overlaps;
 		UnsafeWorld->OverlapMultiByChannel(Overlaps, Bounds.Center, FQuat::Identity, OverlapChannel, FCollisionShape::MakeSphere(Bounds.W), QueryParams, FCollisionResponseParams(ECR_Overlap));
 
+		// @todo(ccaulfield): is there an engine-independent way to do this?
+#if WITH_PHYSX && PHYSICS_INTERFACE_PHYSX
 		SCOPED_SCENE_READ_LOCK(PhysScene ? PhysScene->GetPxScene() : nullptr); //TODO: expose this part to the anim node
-	
+#endif
+
 		for (const FOverlapResult& Overlap : Overlaps)
 		{
 			if (UPrimitiveComponent* OverlapComp = Overlap.GetComponent())
@@ -1137,28 +966,17 @@ void FAnimNode_RigidBody_Chaos::UpdateInternal(const FAnimationUpdateContext& Co
 				if (ComponentsInSim.Contains(OverlapComp) == false)
 				{
 					ComponentsInSim.Add(OverlapComp);
-#if WITH_CHAOS || WITH_IMMEDIATE_PHYSX
-					check(false);
-#else
-					if (PxRigidActor* RigidActor = FPhysicsInterface_PhysX::GetPxRigidActor_AssumesLocked(OverlapComp->BodyInstance.ActorHandle))
-					{
-						PhysicsSimulation->CreateStaticActor(RigidActor, P2UTransform(RigidActor->getGlobalPose()));
-					}
-#endif
+					PhysicsSimulation->CreateActor(ImmediatePhysics_Chaos::EActorType::StaticActor, &OverlapComp->BodyInstance, OverlapComp->BodyInstance.GetUnrealWorldTransform());
 				}
 			}
 		}
-#endif
-
 		UnsafeWorld = nullptr;
 		PhysScene = nullptr;
 	}
-#endif
 }
 
 void FAnimNode_RigidBody_Chaos::InitializeBoneReferences(const FBoneContainer& RequiredBones)
 {
-#if CHAOS_RBAN_TODO
 	/** We only need to update simulated bones and children of simulated bones*/
 	const int32 NumBodies = Bodies.Num();
 	const TArray<FBoneIndexType>& RequiredBoneIndices = RequiredBones.GetBoneIndicesArray();
@@ -1230,7 +1048,6 @@ void FAnimNode_RigidBody_Chaos::InitializeBoneReferences(const FBoneContainer& R
 	// We're switching to a new LOD, this invalidates our captured poses.
 	CapturedFrozenPose.Empty();
 	CapturedFrozenCurves.Empty();
-#endif
 }
 
 void FAnimNode_RigidBody_Chaos::OnInitializeAnimInstance(const FAnimInstanceProxy* InProxy, const UAnimInstance* InAnimInstance)
@@ -1238,9 +1055,22 @@ void FAnimNode_RigidBody_Chaos::OnInitializeAnimInstance(const FAnimInstanceProx
 	InitPhysics(InAnimInstance);
 }
 
+bool FAnimNode_RigidBody_Chaos::IsValidToEvaluate(const USkeleton* Skeleton, const FBoneContainer& RequiredBones)
+{
+	return BaseBoneRef.IsValidToEvaluate(RequiredBones);
+}
+
+#endif // INCLUDE_CHAOS
+
 #if WITH_EDITORONLY_DATA
 void FAnimNode_RigidBody_Chaos::PostSerialize(const FArchive& Ar)
 {
+	if (bComponentSpaceSimulation_DEPRECATED == false)
+	{
+		//If this is not the default value it means we have old content where we were simulating in world space
+		SimulationSpace = ESimulationSpace::WorldSpace;
+		bComponentSpaceSimulation_DEPRECATED = true;
+	}
 }
 #endif
 
diff --git a/Engine/Source/Runtime/AnimGraphRuntime/Public/BoneControllers/AnimNode_RigidBody.h b/Engine/Source/Runtime/AnimGraphRuntime/Public/BoneControllers/AnimNode_RigidBody.h
--- a/Engine/Source/Runtime/AnimGraphRuntime/Public/BoneControllers/AnimNode_RigidBody.h
+++ b/Engine/Source/Runtime/AnimGraphRuntime/Public/BoneControllers/AnimNode_RigidBody.h
@@ -4,14 +4,9 @@
 
 #include "BoneControllers/AnimNode_SkeletalControlBase.h"
 #include "Components/SkeletalMeshComponent.h"
+#include "Physics/ImmediatePhysics/ImmediatePhysicsDeclares.h"
 #include "AnimNode_RigidBody.generated.h"
 
-namespace ImmediatePhysics
-{
-	struct FSimulation;
-	struct FActorHandle;
-}
-
 struct FBodyInstance;
 struct FConstraintInstance;
 
@@ -120,7 +115,7 @@ private:
 	ETeleportType ResetSimulatedTeleportType;
 
 public:
-	UPROPERTY(EditAnywhere, Category = Settings, meta=(InlineEditConditionToggle))
+	UPROPERTY(EditAnywhere, Category = Settings, meta = (InlineEditConditionToggle))
 	uint8 bEnableWorldGeometry : 1;
 
 	UPROPERTY(EditAnywhere, Category = Settings, meta = (InlineEditConditionToggle))
