Remove unused variables

diff --git a/Engine/Source/Runtime/AnimGraphRuntime/Private/BoneControllers/AnimNode_Trail.cpp b/Engine/Source/Runtime/AnimGraphRuntime/Private/BoneControllers/AnimNode_Trail.cpp
--- a/Engine/Source/Runtime/AnimGraphRuntime/Private/BoneControllers/AnimNode_Trail.cpp
+++ b/Engine/Source/Runtime/AnimGraphRuntime/Private/BoneControllers/AnimNode_Trail.cpp
@@ -28,7 +28,6 @@ FAnimNode_Trail::FAnimNode_Trail()
 	, TrailRelaxation_DEPRECATED(10.f)
 #endif// #if WITH_EDITORONLY_DATA
 	, MaxDeltaTime(0.f)
-	, UnwindingSize(3)
 	, RelaxationSpeedScale(1.f)
 	, StretchLimit(0)
 	, FakeVelocity(FVector::ZeroVector)
diff --git a/Engine/Source/Runtime/AnimGraphRuntime/Public/BoneControllers/AnimNode_Trail.h b/Engine/Source/Runtime/AnimGraphRuntime/Public/BoneControllers/AnimNode_Trail.h
--- a/Engine/Source/Runtime/AnimGraphRuntime/Public/BoneControllers/AnimNode_Trail.h
+++ b/Engine/Source/Runtime/AnimGraphRuntime/Public/BoneControllers/AnimNode_Trail.h
@@ -122,13 +122,6 @@ struct ANIMGRAPHRUNTIME_API FAnimNode_Trail : public FAnimNode_SkeletalControlBa
 	UPROPERTY(EditAnywhere, Category = Limit)
 	float MaxDeltaTime;
 
-	/** If you want to avoid loop, how many you want to unwind at once. 
-	 * Bigger value can cause jitter as it becomes more unstable in the ordering
-	 * Defaulted to 3. It will use this length to unwind at once 
-	 */
-	UPROPERTY(EditAnywhere, Category = Trail, meta = (EditCondition = "!bAllowLoop", ClampMin = "3", UIMin = "3"))
-	uint32 UnwindingSize;
-
 	UPROPERTY(EditAnywhere, BlueprintReadWrite, Category = Trail, meta = (PinHiddenByDefault))
 	float RelaxationSpeedScale;
 
