Removed allocation of ImmediateMode::FSimulation when p.RigidBodyNode=0.
This saves 350KB per disabled AnimNode_RigidBody.

This also involves moving p.RigidBodyNode out of scalability settings so it can be read only. Settings for p.RigidBodyNode were moved from scalability settings to device profiles.

diff --git a/Engine/Source/Runtime/AnimGraphRuntime/Private/BoneControllers/AnimNode_RigidBody.cpp b/Engine/Source/Runtime/AnimGraphRuntime/Private/BoneControllers/AnimNode_RigidBody.cpp
--- a/Engine/Source/Runtime/AnimGraphRuntime/Private/BoneControllers/AnimNode_RigidBody.cpp
+++ b/Engine/Source/Runtime/AnimGraphRuntime/Private/BoneControllers/AnimNode_RigidBody.cpp
@@ -21,7 +21,7 @@
 
 #define LOCTEXT_NAMESPACE "ImmediatePhysics"
 
-TAutoConsoleVariable<int32> CVarEnableRigidBodyNode(TEXT("p.RigidBodyNode"), 1, TEXT("Enables/disables rigid body node updates and evaluations"), ECVF_Scalability);
+TAutoConsoleVariable<int32> CVarEnableRigidBodyNode(TEXT("p.RigidBodyNode"), 1, TEXT("Enables/disables rigid body node updates and evaluations"), ECVF_ReadOnly);
 TAutoConsoleVariable<int32> CVarRigidBodyLODThreshold(TEXT("p.RigidBodyLODThreshold"), -1, TEXT("Max LOD that rigid body node is allowed to run on. Provides a global threshold that overrides per-node the LODThreshold property. -1 means no override."), ECVF_Scalability);
 
 FAnimNode_RigidBody::FAnimNode_RigidBody():
@@ -607,7 +607,7 @@ void FAnimNode_RigidBody::InitPhysics(const UAnimInstance* InAnimInstance)
 		bSimulateAnimPhysicsAfterReset = false;
 	}
 
-	if(UsePhysicsAsset)
+	if(CVarEnableRigidBodyNode.GetValueOnAnyThread() != 0 && UsePhysicsAsset)
 	{
 		delete PhysicsSimulation;
 		PhysicsSimulation = new ImmediatePhysics::FSimulation();
