RigidBody Anim Node: Remove unnecessary transform initialization for PhysX

diff --git a/Engine/Source/Runtime/AnimGraphRuntime/Private/BoneControllers/AnimNode_RigidBody.cpp b/Engine/Source/Runtime/AnimGraphRuntime/Private/BoneControllers/AnimNode_RigidBody.cpp
--- a/Engine/Source/Runtime/AnimGraphRuntime/Private/BoneControllers/AnimNode_RigidBody.cpp
+++ b/Engine/Source/Runtime/AnimGraphRuntime/Private/BoneControllers/AnimNode_RigidBody.cpp
@@ -641,7 +641,13 @@ void FAnimNode_RigidBody::InitPhysics(const UAnimInstance* InAnimInstance)
 		// Instantiate a FBodyInstance/FConstraintInstance set that will be cloned into the Immediate Physics sim.
 		TArray<FBodyInstance*> HighLevelBodyInstances;
 		TArray<FConstraintInstance*> HighLevelConstraintInstances;
+#if WITH_CHAOS
+		// Chaos relies on the initial pose to set up constraint positions
 		SkeletalMeshComp->InstantiatePhysicsAssetRefPose(*UsePhysicsAsset, SimulationSpace == ESimulationSpace::WorldSpace ? SkeletalMeshComp->GetComponentToWorld().GetScale3D() : FVector(1.f), HighLevelBodyInstances, HighLevelConstraintInstances);
+#else
+		// PhysX sets up constraints each tick
+		SkeletalMeshComp->InstantiatePhysicsAsset(*UsePhysicsAsset, SimulationSpace == ESimulationSpace::WorldSpace ? SkeletalMeshComp->GetComponentToWorld().GetScale3D() : FVector(1.f), HighLevelBodyInstances, HighLevelConstraintInstances);
+#endif
 
 		TMap<FName, ImmediatePhysics::FActorHandle*> NamesToHandles;
 		TArray<ImmediatePhysics::FActorHandle*> IgnoreCollisionActors;
