Scoped stat counter for AnimNode_RigidBody init time

diff --git a/Engine/Source/Runtime/AnimGraphRuntime/Private/BoneControllers/AnimNode_RigidBody.cpp b/Engine/Source/Runtime/AnimGraphRuntime/Private/BoneControllers/AnimNode_RigidBody.cpp
--- a/Engine/Source/Runtime/AnimGraphRuntime/Private/BoneControllers/AnimNode_RigidBody.cpp
+++ b/Engine/Source/Runtime/AnimGraphRuntime/Private/BoneControllers/AnimNode_RigidBody.cpp
@@ -21,6 +21,8 @@
 
 #define LOCTEXT_NAMESPACE "ImmediatePhysics"
 
+DEFINE_STAT(STAT_RigidBodyNodeInitTime);
+
 TAutoConsoleVariable<int32> CVarEnableRigidBodyNode(TEXT("p.RigidBodyNode"), 1, TEXT("Enables/disables rigid body node updates and evaluations"), ECVF_ReadOnly);
 TAutoConsoleVariable<int32> CVarRigidBodyLODThreshold(TEXT("p.RigidBodyLODThreshold"), -1, TEXT("Max LOD that rigid body node is allowed to run on. Provides a global threshold that overrides per-node the LODThreshold property. -1 means no override."), ECVF_Scalability);
 
@@ -596,6 +598,8 @@ void ComputeBodyInsertionOrder(TArray<FBoneIndexType>& InsertionOrder, const USk
 
 void FAnimNode_RigidBody::InitPhysics(const UAnimInstance* InAnimInstance)
 {
+	SCOPE_CYCLE_COUNTER(STAT_RigidBodyNodeInitTime);
+
 	const USkeletalMeshComponent* SkeletalMeshComp = InAnimInstance->GetSkelMeshComponent();
 	const USkeletalMesh* SkeletalMeshAsset = SkeletalMeshComp->SkeletalMesh;
 
diff --git a/Engine/Source/Runtime/AnimGraphRuntime/Private/BoneControllers/AnimNode_RigidBody_Chaos.cpp b/Engine/Source/Runtime/AnimGraphRuntime/Private/BoneControllers/AnimNode_RigidBody_Chaos.cpp
--- a/Engine/Source/Runtime/AnimGraphRuntime/Private/BoneControllers/AnimNode_RigidBody_Chaos.cpp
+++ b/Engine/Source/Runtime/AnimGraphRuntime/Private/BoneControllers/AnimNode_RigidBody_Chaos.cpp
@@ -570,6 +570,8 @@ void ComputeBodyInsertionOrderChaos(TArray<FBoneIndexType>& InsertionOrder, cons
 
 void FAnimNode_RigidBody_Chaos::InitPhysics(const UAnimInstance* InAnimInstance)
 {
+	SCOPE_CYCLE_COUNTER(STAT_RigidBodyNodeInitTime);
+
 	const USkeletalMeshComponent* SkeletalMeshComp = InAnimInstance->GetSkelMeshComponent();
 	const USkeletalMesh* SkeletalMeshAsset = SkeletalMeshComp->SkeletalMesh;
 
diff --git a/Engine/Source/Runtime/Engine/Public/Physics/ImmediatePhysics/ImmediatePhysicsStats.h b/Engine/Source/Runtime/Engine/Public/Physics/ImmediatePhysics/ImmediatePhysicsStats.h
--- a/Engine/Source/Runtime/Engine/Public/Physics/ImmediatePhysics/ImmediatePhysicsStats.h
+++ b/Engine/Source/Runtime/Engine/Public/Physics/ImmediatePhysics/ImmediatePhysicsStats.h
@@ -5,3 +5,5 @@
 #include "Stats/Stats2.h"
 
 DECLARE_STATS_GROUP(TEXT("Immediate Physics"), STATGROUP_ImmediatePhysics, STATCAT_Advanced);
+
+DECLARE_CYCLE_STAT_EXTERN(TEXT("RigidBodyNodeInitTime"), STAT_RigidBodyNodeInitTime, STATGROUP_ImmediatePhysics,);
