Added SkeletalMeshComponent option to disable rigid body anim nodes.  Enables more selective rban use to save on memory and anim init time for distant meshes.

diff --git a/Engine/Source/Runtime/AnimGraphRuntime/Private/BoneControllers/AnimNode_RigidBody.cpp b/Engine/Source/Runtime/AnimGraphRuntime/Private/BoneControllers/AnimNode_RigidBody.cpp
--- a/Engine/Source/Runtime/AnimGraphRuntime/Private/BoneControllers/AnimNode_RigidBody.cpp
+++ b/Engine/Source/Runtime/AnimGraphRuntime/Private/BoneControllers/AnimNode_RigidBody.cpp
@@ -23,7 +23,7 @@
 
 DEFINE_STAT(STAT_RigidBodyNodeInitTime);
 
-TAutoConsoleVariable<int32> CVarEnableRigidBodyNode(TEXT("p.RigidBodyNode"), 1, TEXT("Enables/disables rigid body node updates and evaluations"), ECVF_ReadOnly);
+TAutoConsoleVariable<int32> CVarEnableRigidBodyNode(TEXT("p.RigidBodyNode"), 1, TEXT("Enables/disables rigid body node updates and evaluations"), ECVF_Default);
 TAutoConsoleVariable<int32> CVarRigidBodyLODThreshold(TEXT("p.RigidBodyLODThreshold"), -1, TEXT("Max LOD that rigid body node is allowed to run on. Provides a global threshold that overrides per-node the LODThreshold property. -1 means no override."), ECVF_Scalability);
 
 int32 RBAN_MaxSubSteps = 4;
@@ -262,7 +262,7 @@ void FAnimNode_RigidBody::EvaluateSkeletalControl_AnyThread(FComponentSpacePoseC
 	const float DeltaSeconds = AccumulatedDeltaTime;
 	AccumulatedDeltaTime = 0.f;
 
-	if (CVarEnableRigidBodyNode.GetValueOnAnyThread() != 0 && PhysicsSimulation)	
+	if (bEnabled && PhysicsSimulation)	
 	{
 		const FBoneContainer& BoneContainer = Output.Pose.GetPose().GetBoneContainer();
 		const FTransform CompWorldSpaceTM = Output.AnimInstanceProxy->GetComponentTransform();
@@ -605,6 +605,9 @@ void FAnimNode_RigidBody::InitPhysics(const UAnimInstance* InAnimInstance)
 {
 	SCOPE_CYCLE_COUNTER(STAT_RigidBodyNodeInitTime);
 
+	delete PhysicsSimulation;
+	PhysicsSimulation = nullptr;
+
 	const USkeletalMeshComponent* SkeletalMeshComp = InAnimInstance->GetSkelMeshComponent();
 	const USkeletalMesh* SkeletalMeshAsset = SkeletalMeshComp->SkeletalMesh;
 
@@ -635,10 +638,10 @@ void FAnimNode_RigidBody::InitPhysics(const UAnimInstance* InAnimInstance)
 		AnimPhysicsMinDeltaTime = 0.f;
 		bSimulateAnimPhysicsAfterReset = false;
 	}
-
-	if(CVarEnableRigidBodyNode.GetValueOnAnyThread() != 0 && UsePhysicsAsset)
+	
+	bEnabled = UsePhysicsAsset && SkeletalMeshComp->GetAllowRigidBodyAnimNode() && CVarEnableRigidBodyNode.GetValueOnAnyThread() != 0;
+	if(bEnabled)
 	{
-		delete PhysicsSimulation;
 		PhysicsSimulation = new ImmediatePhysics::FSimulation();
 		const int32 NumBodies = UsePhysicsAsset->SkeletalBodySetups.Num();
 		Bodies.Empty(NumBodies);
@@ -915,7 +918,7 @@ DECLARE_CYCLE_STAT(TEXT("RigidBody_PreUpdate"), STAT_RigidBody_PreUpdate, STATGR
 void FAnimNode_RigidBody::PreUpdate(const UAnimInstance* InAnimInstance)
 {
 	// Don't update geometry if RBN is disabled
-	if(CVarEnableRigidBodyNode.GetValueOnAnyThread() == 0)
+	if(!bEnabled)
 	{
 		return;
 	}
@@ -975,7 +978,7 @@ void FAnimNode_RigidBody::UpdateInternal(const FAnimationUpdateContext& Context)
 {
 	DECLARE_SCOPE_HIERARCHICAL_COUNTER_ANIMNODE(UpdateInternal)
 	// Avoid this work if RBN is disabled, as the results would be discarded
-	if(CVarEnableRigidBodyNode.GetValueOnAnyThread() == 0)
+	if(!bEnabled)
 	{
 		return;
 	}
diff --git a/Engine/Source/Runtime/AnimGraphRuntime/Private/BoneControllers/AnimNode_RigidBody_Chaos.cpp b/Engine/Source/Runtime/AnimGraphRuntime/Private/BoneControllers/AnimNode_RigidBody_Chaos.cpp
--- a/Engine/Source/Runtime/AnimGraphRuntime/Private/BoneControllers/AnimNode_RigidBody_Chaos.cpp
+++ b/Engine/Source/Runtime/AnimGraphRuntime/Private/BoneControllers/AnimNode_RigidBody_Chaos.cpp
@@ -236,7 +236,7 @@ void FAnimNode_RigidBody_Chaos::EvaluateSkeletalControl_AnyThread(FComponentSpac
 	const float DeltaSeconds = AccumulatedDeltaTime;
 	AccumulatedDeltaTime = 0.f;
 
-	if (CVarEnableRigidBodyNode.GetValueOnAnyThread() != 0 && PhysicsSimulation)	
+	if (bEnabled && PhysicsSimulation)	
 	{
 		const FBoneContainer& BoneContainer = Output.Pose.GetPose().GetBoneContainer();
 		const FTransform CompWorldSpaceTM = Output.AnimInstanceProxy->GetComponentTransform();
@@ -572,6 +572,9 @@ void FAnimNode_RigidBody_Chaos::InitPhysics(const UAnimInstance* InAnimInstance)
 {
 	SCOPE_CYCLE_COUNTER(STAT_RigidBodyNodeInitTime);
 
+	delete PhysicsSimulation;
+	PhysicsSimulation = nullptr;
+
 	const USkeletalMeshComponent* SkeletalMeshComp = InAnimInstance->GetSkelMeshComponent();
 	const USkeletalMesh* SkeletalMeshAsset = SkeletalMeshComp->SkeletalMesh;
 
@@ -603,9 +606,9 @@ void FAnimNode_RigidBody_Chaos::InitPhysics(const UAnimInstance* InAnimInstance)
 		bSimulateAnimPhysicsAfterReset = false;
 	}
 
-	if(UsePhysicsAsset)
+	bEnabled = UsePhysicsAsset && SkeletalMeshComp->GetAllowRigidBodyAnimNode() && CVarEnableRigidBodyNode.GetValueOnAnyThread() != 0;
+	if (bEnabled)
 	{
-		delete PhysicsSimulation;
 		PhysicsSimulation = new ImmediatePhysics_Chaos::FSimulation();
 		const int32 NumBodies = UsePhysicsAsset->SkeletalBodySetups.Num();
 		Bodies.Empty(NumBodies);
@@ -876,7 +879,7 @@ DECLARE_CYCLE_STAT(TEXT("RigidBody_PreUpdate_Chaos"), STAT_RigidBody_Chaos_PreUp
 void FAnimNode_RigidBody_Chaos::PreUpdate(const UAnimInstance* InAnimInstance)
 {
 	// Don't update geometry if RBN is disabled
-	if(CVarEnableRigidBodyNode.GetValueOnAnyThread() == 0)
+	if(!bEnabled)
 	{
 		return;
 	}
@@ -935,7 +938,7 @@ DECLARE_CYCLE_STAT(TEXT("RigidBody_Update_Chaos"), STAT_RigidBody_Chaos_Update,
 void FAnimNode_RigidBody_Chaos::UpdateInternal(const FAnimationUpdateContext& Context)
 {
 	// Avoid this work if RBN is disabled, as the results would be discarded
-	if(CVarEnableRigidBodyNode.GetValueOnAnyThread() == 0)
+	if(!bEnabled)
 	{
 		return;
 	}
diff --git a/Engine/Source/Runtime/AnimGraphRuntime/Public/BoneControllers/AnimNode_RigidBody.h b/Engine/Source/Runtime/AnimGraphRuntime/Public/BoneControllers/AnimNode_RigidBody.h
--- a/Engine/Source/Runtime/AnimGraphRuntime/Public/BoneControllers/AnimNode_RigidBody.h
+++ b/Engine/Source/Runtime/AnimGraphRuntime/Public/BoneControllers/AnimNode_RigidBody.h
@@ -145,6 +145,7 @@ public:
 	uint8 bClampLinearTranslationLimitToRefPose : 1;
 
 private:
+	uint8 bEnabled : 1;
 	uint8 bSimulationStarted : 1;
 	uint8 bCheckForBodyTransformInit : 1;
 
diff --git a/Engine/Source/Runtime/AnimGraphRuntime/Public/BoneControllers/AnimNode_RigidBody_Chaos.h b/Engine/Source/Runtime/AnimGraphRuntime/Public/BoneControllers/AnimNode_RigidBody_Chaos.h
--- a/Engine/Source/Runtime/AnimGraphRuntime/Public/BoneControllers/AnimNode_RigidBody_Chaos.h
+++ b/Engine/Source/Runtime/AnimGraphRuntime/Public/BoneControllers/AnimNode_RigidBody_Chaos.h
@@ -155,6 +155,7 @@ private:
 
 	ETeleportType ResetSimulatedTeleportType;
 
+	uint8 bEnabled : 1;
 	uint8 bSimulationStarted : 1;
 	uint8 bCheckForBodyTransformInit : 1;
 
diff --git a/Engine/Source/Runtime/Engine/Classes/Components/SkeletalMeshComponent.h b/Engine/Source/Runtime/Engine/Classes/Components/SkeletalMeshComponent.h
--- a/Engine/Source/Runtime/Engine/Classes/Components/SkeletalMeshComponent.h
+++ b/Engine/Source/Runtime/Engine/Classes/Components/SkeletalMeshComponent.h
@@ -473,6 +473,10 @@ public:
 	uint8 bDisableClothSimulation:1;
 
 private:
+	/** Disable rigid body animation nodes and play original animation without simulation */
+	UPROPERTY(EditAnywhere, Category=Clothing)
+	uint8 bDisableRigidBodyAnimNode:1;
+
 	/** Disable animation curves for this component. If this is set true, no curves will be processed */
 	UPROPERTY(EditAnywhere, AdvancedDisplay, Category = SkeletalMesh)
 	uint8 bAllowAnimCurveEvaluation : 1;
@@ -1055,6 +1059,15 @@ public:
 	UFUNCTION(BlueprintCallable, Category="Components|SkeletalMesh", meta=(UnsafeDuringActorConstruction="true"))
 	void UnbindClothFromMasterPoseComponent(bool bRestoreSimulationSpace = true);
 
+	/**
+	 * Sets whether or not to allow rigid body animation nodes for this component
+	 */
+	UFUNCTION(BlueprintCallable, Category = "Components|SkeletalMesh")
+	void SetAllowRigidBodyAnimNode(bool bInAllow, bool bReinitAnim = true);
+
+	UFUNCTION(BlueprintCallable, Category = "Components|SkeletalMesh")
+	bool GetAllowRigidBodyAnimNode() const { return !bDisableRigidBodyAnimNode; }
+
 	/**
 	* Sets whether or not to force tick component in order to update animation and refresh transform for this component
 	* This is supported only in the editor
diff --git a/Engine/Source/Runtime/Engine/Private/Components/SkeletalMeshComponent.cpp b/Engine/Source/Runtime/Engine/Private/Components/SkeletalMeshComponent.cpp
--- a/Engine/Source/Runtime/Engine/Private/Components/SkeletalMeshComponent.cpp
+++ b/Engine/Source/Runtime/Engine/Private/Components/SkeletalMeshComponent.cpp
@@ -3660,8 +3660,22 @@ void USkeletalMeshComponent::UnbindClothFromMasterPoseComponent(bool bRestoreSim
 		}
 
 		bBindClothToMasterComponent = false;
+	}
+}
+
+void USkeletalMeshComponent::SetAllowRigidBodyAnimNode(bool bInAllow, bool bReinitAnim)
+{
+	if(bDisableRigidBodyAnimNode == bInAllow)
+	{
+		bDisableRigidBodyAnimNode = !bInAllow;
+
+		if(bReinitAnim && bRegistered && SkeletalMesh)
+		{
+			// need to reinitialize rigid body nodes for new setting to take effect
+			InitializeAnimScriptInstance(true);
 		}
 	}
+}
 
 bool USkeletalMeshComponent::DoCustomNavigableGeometryExport(FNavigableGeometryExport& GeomExport) const
 {
