Sequencer: Correctly clear the active anim instance on a skeletal animation component when stopping a sequence.

This will stop any active anim-notifies when reverting to pre-animated state.

diff --git a/Engine/Source/Runtime/AnimGraphRuntime/Private/AnimCustomInstance.cpp b/Engine/Source/Runtime/AnimGraphRuntime/Private/AnimCustomInstance.cpp
--- a/Engine/Source/Runtime/AnimGraphRuntime/Private/AnimCustomInstance.cpp
+++ b/Engine/Source/Runtime/AnimGraphRuntime/Private/AnimCustomInstance.cpp
@@ -27,7 +27,7 @@ void UAnimCustomInstance::UnbindFromSkeletalMeshComponent(USkeletalMeshComponent
 		UAnimCustomInstance* SequencerInstance = Cast<UAnimCustomInstance>(InSkeletalMeshComponent->GetAnimInstance());
 		if (SequencerInstance)
 		{
-			InSkeletalMeshComponent->AnimScriptInstance = nullptr;
+			InSkeletalMeshComponent->ClearAnimScriptInstance();
 		}
 	}
 	else if (InSkeletalMeshComponent->GetAnimationMode() == EAnimationMode::Type::AnimationBlueprint)
@@ -60,4 +60,4 @@ bool UAnimCustomInstance::ShouldUseSequenceInstancePlayer(const USkeletalMeshCom
 
 	return (AnimInstance == nullptr || SkeletalMeshComponent->GetAnimationMode() != EAnimationMode::AnimationBlueprint ||
 		AnimInstance->GetClass() != SkeletalMeshComponent->AnimClass || !SkeletalMesh->Skeleton->IsCompatible(AnimInstance->CurrentSkeleton));
-}
\ No newline at end of file
+}
