Fix CIS for AnimNode_RigidBody_Chaos

diff --git a/Engine/Source/Runtime/AnimGraphRuntime/Private/BoneControllers/AnimNode_RigidBody_Chaos.cpp b/Engine/Source/Runtime/AnimGraphRuntime/Private/BoneControllers/AnimNode_RigidBody_Chaos.cpp
--- a/Engine/Source/Runtime/AnimGraphRuntime/Private/BoneControllers/AnimNode_RigidBody_Chaos.cpp
+++ b/Engine/Source/Runtime/AnimGraphRuntime/Private/BoneControllers/AnimNode_RigidBody_Chaos.cpp
@@ -447,19 +447,7 @@ void FAnimNode_RigidBody_Chaos::EvaluateSkeletalControl_AnyThread(FComponentSpac
 			extern int32 RBAN_MaxSubSteps;
 			const int32 MaxSteps = RBAN_MaxSubSteps;
 			const float MaxDeltaSeconds = 1.f / 30.f;
-
-#if !WITH_CHAOS
-			const int32 NumSteps = FMath::Clamp(FMath::CeilToInt(DeltaSeconds / MaxDeltaSeconds), 1, MaxSteps);
-			const float StepDeltaTime = DeltaSeconds / float(NumSteps);
-			for (int32 Step = 1; Step <= NumSteps; Step++)
-			{
-				// We call the _AssumesLocked version here without a lock as the simulation is local to this node and we know
-				// we're not going to alter anything while this is running.
-				PhysicsSimulation->Simulate_AssumesLocked(StepDeltaTime, SimSpaceGravity);
-			}
-#else
 			PhysicsSimulation->Simulate_AssumesLocked(DeltaSeconds, MaxDeltaSeconds, MaxSteps, SimSpaceGravity);
-#endif
 		}
 		
 		//write back to animation system
