Updated and expanded CSV Animation stats

diff --git a/Engine/Source/Runtime/AnimGraphRuntime/Private/BoneControllers/AnimNode_AnimDynamics.cpp b/Engine/Source/Runtime/AnimGraphRuntime/Private/BoneControllers/AnimNode_AnimDynamics.cpp
--- a/Engine/Source/Runtime/AnimGraphRuntime/Private/BoneControllers/AnimNode_AnimDynamics.cpp
+++ b/Engine/Source/Runtime/AnimGraphRuntime/Private/BoneControllers/AnimNode_AnimDynamics.cpp
@@ -11,6 +11,8 @@ DEFINE_STAT(STAT_AnimDynamicsWindData);
 DEFINE_STAT(STAT_AnimDynamicsBoneEval);
 DEFINE_STAT(STAT_AnimDynamicsSubSteps);
 
+CSV_DECLARE_CATEGORY_MODULE_EXTERN(ENGINE_API, Animation);
+
 TAutoConsoleVariable<int32> CVarRestrictLod(TEXT("p.AnimDynamicsRestrictLOD"), -1, TEXT("Forces anim dynamics to be enabled for only a specified LOD, -1 to enable on all LODs."));
 TAutoConsoleVariable<int32> CVarLODThreshold(TEXT("p.AnimDynamicsLODThreshold"), -1, TEXT("Max LOD that anim dynamics is allowed to run on. Provides a global threshold that overrides per-node the LODThreshold property. -1 means no override."), ECVF_Scalability);
 TAutoConsoleVariable<int32> CVarEnableDynamics(TEXT("p.AnimDynamics"), 1, TEXT("Enables/Disables anim dynamics node updates."), ECVF_Scalability);
@@ -167,6 +169,7 @@ FAnimNode_AnimDynamics::FAnimNode_AnimDynamics()
 
 void FAnimNode_AnimDynamics::Initialize_AnyThread(const FAnimationInitializeContext& Context)
 {
+	CSV_SCOPED_TIMING_STAT(Animation, AnimDynamicsInit);
 	DECLARE_SCOPE_HIERARCHICAL_COUNTER_ANIMNODE(Initialize_AnyThread)
 	FAnimNode_SkeletalControlBase::Initialize_AnyThread(Context);
 
@@ -212,6 +215,7 @@ void FAnimNode_AnimDynamics::EvaluateSkeletalControl_AnyThread(FComponentSpacePo
 {
 	DECLARE_SCOPE_HIERARCHICAL_COUNTER_ANIMNODE(EvaluateSkeletalControl_AnyThread)
 	SCOPE_CYCLE_COUNTER(STAT_AnimDynamicsOverall);
+	CSV_SCOPED_TIMING_STAT(Animation, AnimDynamicsEval);
 
 	if (IsAnimDynamicsSystemEnabledFor(Output.AnimInstanceProxy->GetLODLevel()))
 	{
diff --git a/Engine/Source/Runtime/AnimGraphRuntime/Private/BoneControllers/AnimNode_RigidBody.cpp b/Engine/Source/Runtime/AnimGraphRuntime/Private/BoneControllers/AnimNode_RigidBody.cpp
--- a/Engine/Source/Runtime/AnimGraphRuntime/Private/BoneControllers/AnimNode_RigidBody.cpp
+++ b/Engine/Source/Runtime/AnimGraphRuntime/Private/BoneControllers/AnimNode_RigidBody.cpp
@@ -23,6 +23,8 @@
 
 DEFINE_STAT(STAT_RigidBodyNodeInitTime);
 
+CSV_DECLARE_CATEGORY_MODULE_EXTERN(ENGINE_API, Animation);
+
 TAutoConsoleVariable<int32> CVarEnableRigidBodyNode(TEXT("p.RigidBodyNode"), 1, TEXT("Enables/disables rigid body node updates and evaluations"), ECVF_Default);
 TAutoConsoleVariable<int32> CVarRigidBodyLODThreshold(TEXT("p.RigidBodyLODThreshold"), -1, TEXT("Max LOD that rigid body node is allowed to run on. Provides a global threshold that overrides per-node the LODThreshold property. -1 means no override."), ECVF_Scalability);
 
@@ -243,6 +245,7 @@ void FAnimNode_RigidBody::EvaluateSkeletalControl_AnyThread(FComponentSpacePoseC
 {
 	DECLARE_SCOPE_HIERARCHICAL_COUNTER_ANIMNODE(EvaluateSkeletalControl_AnyThread)
 	SCOPE_CYCLE_COUNTER(STAT_RigidBody_Eval);
+	CSV_SCOPED_TIMING_STAT(Animation, RigidBodyEval);
 	FScopeCycleCounterUObject AdditionalScope(UsePhysicsAsset, GET_STATID(STAT_RigidBody_Eval));
 	SCOPE_CYCLE_COUNTER(STAT_ImmediateEvaluateSkeletalControl);
 	//SCOPED_NAMED_EVENT_TEXT("FAnimNode_RigidBody::EvaluateSkeletalControl_AnyThread", FColor::Magenta);
@@ -1099,6 +1102,7 @@ void FAnimNode_RigidBody::InitializeBoneReferences(const FBoneContainer& Require
 
 void FAnimNode_RigidBody::OnInitializeAnimInstance(const FAnimInstanceProxy* InProxy, const UAnimInstance* InAnimInstance)
 {
+	CSV_SCOPED_TIMING_STAT(Animation, RigidBodyInit);
 	InitPhysics(InAnimInstance);
 }
 
diff --git a/Engine/Source/Runtime/Engine/Private/Animation/AnimCompressionTypes.cpp b/Engine/Source/Runtime/Engine/Private/Animation/AnimCompressionTypes.cpp
--- a/Engine/Source/Runtime/Engine/Private/Animation/AnimCompressionTypes.cpp
+++ b/Engine/Source/Runtime/Engine/Private/Animation/AnimCompressionTypes.cpp
@@ -10,7 +10,7 @@
 #include "AnimationRuntime.h"
 #include "UObject/FortniteMainBranchObjectVersion.h"
 
-CSV_DEFINE_CATEGORY(Animation, false);
+CSV_DECLARE_CATEGORY_EXTERN(Animation);
 
 DECLARE_CYCLE_STAT(TEXT("Build Anim Track Pairs"), STAT_BuildAnimTrackPairs, STATGROUP_Anim);
 DECLARE_CYCLE_STAT(TEXT("Extract Pose From Anim Data"), STAT_ExtractPoseFromAnimData, STATGROUP_Anim);
diff --git a/Engine/Source/Runtime/Engine/Private/Animation/AnimInstance.cpp b/Engine/Source/Runtime/Engine/Private/Animation/AnimInstance.cpp
--- a/Engine/Source/Runtime/Engine/Private/Animation/AnimInstance.cpp
+++ b/Engine/Source/Runtime/Engine/Private/Animation/AnimInstance.cpp
@@ -74,6 +74,8 @@ DEFINE_STAT(STAT_AnimGameThreadTime);
 DEFINE_STAT(STAT_TickAssetPlayerInstances);
 DEFINE_STAT(STAT_TickAssetPlayerInstance);
 
+CSV_DEFINE_CATEGORY_MODULE(ENGINE_API, Animation, false);
+
 // Define AnimNotify
 DEFINE_LOG_CATEGORY(LogAnimNotify);
 
diff --git a/Engine/Source/Runtime/Engine/Private/Animation/AnimSequence.cpp b/Engine/Source/Runtime/Engine/Private/Animation/AnimSequence.cpp
--- a/Engine/Source/Runtime/Engine/Private/Animation/AnimSequence.cpp
+++ b/Engine/Source/Runtime/Engine/Private/Animation/AnimSequence.cpp
@@ -46,6 +46,8 @@
 DECLARE_CYCLE_STAT(TEXT("AnimSeq GetBonePose"), STAT_AnimSeq_GetBonePose, STATGROUP_Anim);
 DECLARE_CYCLE_STAT(TEXT("AnimSeq EvalCurveData"), STAT_AnimSeq_EvalCurveData, STATGROUP_Anim);
 
+CSV_DECLARE_CATEGORY_EXTERN(Animation);
+
 int32 GPerformFrameStripping = 0;
 int32 GPerformFrameStrippingOddFramedAnimations = 0;
 
diff --git a/Engine/Source/Runtime/Engine/Private/Components/SkeletalMeshComponent.cpp b/Engine/Source/Runtime/Engine/Private/Components/SkeletalMeshComponent.cpp
--- a/Engine/Source/Runtime/Engine/Private/Components/SkeletalMeshComponent.cpp
+++ b/Engine/Source/Runtime/Engine/Private/Components/SkeletalMeshComponent.cpp
@@ -57,6 +57,7 @@ DECLARE_CYCLE_STAT_EXTERN(TEXT("Anim Instance Spawn Time"), STAT_AnimSpawnTime,
 DEFINE_STAT(STAT_AnimSpawnTime);
 DEFINE_STAT(STAT_PostAnimEvaluation);
 
+CSV_DECLARE_CATEGORY_EXTERN(Animation);
 CSV_DECLARE_CATEGORY_MODULE_EXTERN(CORE_API, Basic);
 
 FAutoConsoleTaskPriority CPrio_ParallelAnimationEvaluationTask(
@@ -606,6 +607,7 @@ void USkeletalMeshComponent::OnUnregister()
 
 void USkeletalMeshComponent::InitAnim(bool bForceReinit)
 {
+	CSV_SCOPED_TIMING_STAT(Animation, InitAnim);
 	QUICK_SCOPE_CYCLE_COUNTER(STAT_SkelMeshComp_InitAnim);
 	LLM_SCOPE(ELLMTag::Animation);
 
@@ -1802,6 +1804,7 @@ void USkeletalMeshComponent::PerformAnimationEvaluation(const USkeletalMesh* InS
 
 void USkeletalMeshComponent::PerformAnimationProcessing(const USkeletalMesh* InSkeletalMesh, UAnimInstance* InAnimInstance, bool bInDoEvaluation, TArray<FTransform>& OutSpaceBases, TArray<FTransform>& OutBoneSpaceTransforms, FVector& OutRootBoneTranslation, FBlendedHeapCurve& OutCurve)
 {
+	CSV_SCOPED_TIMING_STAT(Animation, WorkerThreadTickTime);
 	ANIM_MT_SCOPE_CYCLE_COUNTER(PerformAnimEvaluation, !IsInGameThread());
 
 	// Can't do anything without a SkeletalMesh
diff --git a/Engine/Source/Runtime/Engine/Private/SkeletalMeshComponentPhysics.cpp b/Engine/Source/Runtime/Engine/Private/SkeletalMeshComponentPhysics.cpp
--- a/Engine/Source/Runtime/Engine/Private/SkeletalMeshComponentPhysics.cpp
+++ b/Engine/Source/Runtime/Engine/Private/SkeletalMeshComponentPhysics.cpp
@@ -58,6 +58,8 @@
 
 DECLARE_CYCLE_STAT(TEXT("CreateClothing"), STAT_CreateClothing, STATGROUP_Physics);
 
+CSV_DECLARE_CATEGORY_EXTERN(Animation);
+
 CSV_DECLARE_CATEGORY_MODULE_EXTERN(CORE_API, Basic);
 
 TAutoConsoleVariable<int32> CVarEnableClothPhysics(TEXT("p.ClothPhysics"), 1, TEXT("If 1, physics cloth will be used for simulation."));
@@ -2327,6 +2329,8 @@ void USkeletalMeshComponent::AddClothingBounds(FBoxSphereBounds& InOutBounds, co
 
 void USkeletalMeshComponent::RecreateClothingActors()
 {
+	CSV_SCOPED_TIMING_STAT(Animation, ClothInit);
+
 	ReleaseAllClothingResources();
 
 	if(SkeletalMesh == nullptr || !IsRegistered())
diff --git a/Engine/Source/Runtime/Engine/Public/Animation/AnimCompressionTypes.h b/Engine/Source/Runtime/Engine/Public/Animation/AnimCompressionTypes.h
--- a/Engine/Source/Runtime/Engine/Public/Animation/AnimCompressionTypes.h
+++ b/Engine/Source/Runtime/Engine/Public/Animation/AnimCompressionTypes.h
@@ -24,8 +24,6 @@
 
 #include "AnimCompressionTypes.generated.h"
 
-CSV_DECLARE_CATEGORY_EXTERN(Animation);
-
 /**
  * Indicates animation data key format.
  */
