# Fix for crash when skeleton has missing bones on rigidbody initialize
#rnx

diff --git a/Engine/Source/Runtime/AnimGraphRuntime/Private/BoneControllers/AnimNode_RigidBody.cpp b/Engine/Source/Runtime/AnimGraphRuntime/Private/BoneControllers/AnimNode_RigidBody.cpp
--- a/Engine/Source/Runtime/AnimGraphRuntime/Private/BoneControllers/AnimNode_RigidBody.cpp
+++ b/Engine/Source/Runtime/AnimGraphRuntime/Private/BoneControllers/AnimNode_RigidBody.cpp
@@ -1055,10 +1055,20 @@ void FAnimNode_RigidBody::InitializeBoneReferences(const FBoneContainer& Require
 		BaseBoneRef.Initialize(RequiredBones);
 	}
 
-	for(int32 Index = 0; Index < NumRequiredBoneIndices; ++Index)
+	bool bHasInvalidBoneReference = false;
+	for (int32 Index = 0; Index < NumRequiredBoneIndices; ++Index)
 	{
 		const FCompactPoseBoneIndex CompactPoseBoneIndex(Index);
 		const FBoneIndexType SkeletonBoneIndex = RequiredBones.GetSkeletonIndex(CompactPoseBoneIndex);
+		const FBoneIndexType IndexToBodyNum = SkeletonBoneIndexToBodyIndex.Num();
+
+		// If we have a missing bone in our skeleton, we don't want to have an out of bounds access.
+		if (SkeletonBoneIndex >= IndexToBodyNum)
+		{
+			bHasInvalidBoneReference = true;
+			break;
+		}
+
 		const int32 BodyIndex = SkeletonBoneIndexToBodyIndex[SkeletonBoneIndex];
 
 		if (BodyIndex != INDEX_NONE)
@@ -1081,6 +1091,14 @@ void FAnimNode_RigidBody::InitializeBoneReferences(const FBoneContainer& Require
 			while (CompactParentIndex != INDEX_NONE)
 			{
 				const FBoneIndexType SkeletonParentBoneIndex = RequiredBones.GetSkeletonIndex(CompactParentIndex);
+
+				// Must check our parent as well for a missing bone.
+				if (SkeletonParentBoneIndex >= IndexToBodyNum)
+				{
+					bHasInvalidBoneReference = true;
+					break;
+				}
+
 				OutputData->ParentBodyIndex = SkeletonBoneIndexToBodyIndex[SkeletonParentBoneIndex];
 				if (OutputData->ParentBodyIndex != INDEX_NONE)
 				{
@@ -1090,21 +1108,36 @@ void FAnimNode_RigidBody::InitializeBoneReferences(const FBoneContainer& Require
 				OutputData->BoneIndicesToParentBody.Add(CompactParentIndex);
 				CompactParentIndex = RequiredBones.GetParentBoneIndex(CompactParentIndex);
 			}
+
+			if (bHasInvalidBoneReference)
+			{
+				break;
+			}
 		}
 	}
 
-	// New bodies potentially introduced with new LOD
-	// We'll have to initialize their transform.
-	bCheckForBodyTransformInit = true;
-
-	if(PhysicsSimulation)
+	if (bHasInvalidBoneReference)
 	{
-		PhysicsSimulation->SetNumActiveBodies(NumSimulatedBodies);
+		// If a bone was missing, let us know which asset it happened on, and clear our bone container to make the bad asset visible.
+		ensureMsgf(false, TEXT("FAnimNode_RigidBody::InitializeBoneReferences: The Skeleton %s, is missing bones that SkeletalMesh %s needs. Skeleton might need to be resaved."),
+			*GetNameSafe(RequiredBones.GetSkeletonAsset()), *GetNameSafe(RequiredBones.GetSkeletalMeshAsset()));
+		OutputBoneData.Empty();
 	}
+	else
+	{
+		// New bodies potentially introduced with new LOD
+		// We'll have to initialize their transform.
+		bCheckForBodyTransformInit = true;
 
-	// We're switching to a new LOD, this invalidates our captured poses.
-	CapturedFrozenPose.Empty();
-	CapturedFrozenCurves.Empty();
+		if (PhysicsSimulation)
+		{
+			PhysicsSimulation->SetNumActiveBodies(NumSimulatedBodies);
+		}
+
+		// We're switching to a new LOD, this invalidates our captured poses.
+		CapturedFrozenPose.Empty();
+		CapturedFrozenCurves.Empty();
+	}
 }
 
 void FAnimNode_RigidBody::OnInitializeAnimInstance(const FAnimInstanceProxy* InProxy, const UAnimInstance* InAnimInstance)
diff --git a/Engine/Source/Runtime/Engine/Private/Animation/Skeleton.cpp b/Engine/Source/Runtime/Engine/Private/Animation/Skeleton.cpp
--- a/Engine/Source/Runtime/Engine/Private/Animation/Skeleton.cpp
+++ b/Engine/Source/Runtime/Engine/Private/Animation/Skeleton.cpp
@@ -514,6 +514,10 @@ int32 USkeleton::BuildLinkup(const USkeletalMesh* InSkelMesh)
 			// Fix missing bone.
 			SkeletonBoneIndex = SkeletonRefSkel.FindBoneIndex(MeshBoneName);
 		}
+#else
+		// If we're not in editor, we still want to know which skeleton is missing a bone.
+		ensureMsgf(SkeletonBoneIndex != INDEX_NONE, TEXT("USkeleton::BuildLinkup: The Skeleton %s, is missing bones that SkeletalMesh %s needs. MeshBoneName %s"),
+				*GetNameSafe(this), *GetNameSafe(InSkelMesh), *MeshBoneName.ToString());
 #endif
 
 		NewMeshLinkup.MeshToSkeletonTable[MeshBoneIndex] = SkeletonBoneIndex;
