Chaos - minot cleanup

diff --git a/Engine/Source/Runtime/AnimGraphRuntime/Private/BoneControllers/AnimNode_RigidBody_Chaos.cpp b/Engine/Source/Runtime/AnimGraphRuntime/Private/BoneControllers/AnimNode_RigidBody_Chaos.cpp
--- a/Engine/Source/Runtime/AnimGraphRuntime/Private/BoneControllers/AnimNode_RigidBody_Chaos.cpp
+++ b/Engine/Source/Runtime/AnimGraphRuntime/Private/BoneControllers/AnimNode_RigidBody_Chaos.cpp
@@ -81,6 +81,21 @@ void FAnimNode_RigidBody_Chaos::GatherDebugData(FNodeDebugData& DebugData)
 	}
 }
 
+FTransform SpaceToWorldTransformChaos(ESimulationSpace Space, const FTransform& ComponentToWorld, const FTransform& BaseBoneTM)
+{
+	switch (Space)
+	{
+	case ESimulationSpace::ComponentSpace:
+		return ComponentToWorld;
+	case ESimulationSpace::WorldSpace:
+		return FTransform::Identity;
+	case ESimulationSpace::BaseBoneSpace:
+		return BaseBoneTM * ComponentToWorld;
+	default:
+		return FTransform::Identity;
+	}
+}
+
 FVector WorldVectorToSpaceNoScaleChaos(ESimulationSpace Space, const FVector& WorldDir, const FTransform& ComponentToWorld, const FTransform& BaseBoneTM)
 {
 	switch(Space)
@@ -247,6 +262,11 @@ void FAnimNode_RigidBody_Chaos::EvaluateSkeletalControl_AnyThread(FComponentSpac
 		}
 		const FTransform BaseBoneTM = Output.Pose.GetComponentSpaceTransform(BaseBoneRef.GetCompactPoseIndex(BoneContainer));
 
+#if !UE_BUILD_SHIPPING
+		// Only used for debug draw...
+		PhysicsSimulation->SetSimulationSpaceTransform(SpaceToWorldTransformChaos(SimulationSpace, CompWorldSpaceTM, BaseBoneTM));
+#endif
+
 		// Initialize potential new bodies because of LOD change.
 		if (ResetSimulatedTeleportType == ETeleportType::None && bCheckForBodyTransformInit)
 		{
diff --git a/Engine/Source/Runtime/Experimental/Chaos/Public/Chaos/DenseMatrix.h b/Engine/Source/Runtime/Experimental/Chaos/Public/Chaos/DenseMatrix.h
--- a/Engine/Source/Runtime/Experimental/Chaos/Public/Chaos/DenseMatrix.h
+++ b/Engine/Source/Runtime/Experimental/Chaos/Public/Chaos/DenseMatrix.h
@@ -784,7 +784,7 @@ namespace Chaos
 				X.At(I, 0) = Sum / G.At(I, I);
 			}
 
-			// Solve LtX = Y (Lt is upper-triangular)
+			// Solve LtX = Y (Gt is upper-triangular)
 			for (int32 I = N - 1; I >= 0; --I)
 			{
 				FReal Sum = X.At(I, 0);
diff --git a/Engine/Source/Runtime/Experimental/Chaos/Public/Chaos/PBDJointConstraints.h b/Engine/Source/Runtime/Experimental/Chaos/Public/Chaos/PBDJointConstraints.h
--- a/Engine/Source/Runtime/Experimental/Chaos/Public/Chaos/PBDJointConstraints.h
+++ b/Engine/Source/Runtime/Experimental/Chaos/Public/Chaos/PBDJointConstraints.h
@@ -143,7 +143,7 @@ namespace Chaos
 		friend class FPBDJointConstraintHandle;
 
 		void CalculateConstraintSpace(int32 ConstraintIndex, FVec3& OutX0, FMatrix33& OutR0, FVec3& OutX1, FMatrix33& OutR1, FVec3& OutAngles) const;
-		void UpdateParticleState(TPBDRigidParticleHandle<FReal, 3>* Rigid, const FReal Dt, const FVec3& P, const FRotation3& Q);
+		void UpdateParticleState(TPBDRigidParticleHandle<FReal, 3>* Rigid, const FReal Dt, const FVec3& P, const FRotation3& Q, const bool bUpdateVelocity = true);
 
 		void ApplyDrives(const FReal Dt, const int32 ConstraintIndex, const int32 NumPairIts, const int32 It, const int32 NumIts);
 		void SolveVelocity(const FReal Dt, const int32 ConstraintIndex, const int32 NumPairIts, const int32 It, const int32 NumIts);
