COPY from //Dev-Anim to //Dev-Main

diff --git a/Engine/Source/Editor/UnrealEd/Private/Factories/SkeletalMeshImport.cpp b/Engine/Source/Editor/UnrealEd/Private/Factories/SkeletalMeshImport.cpp
--- a/Engine/Source/Editor/UnrealEd/Private/Factories/SkeletalMeshImport.cpp
+++ b/Engine/Source/Editor/UnrealEd/Private/Factories/SkeletalMeshImport.cpp
@@ -38,9 +38,9 @@ DEFINE_LOG_CATEGORY_STATIC(LogSkeletalMeshImport, Log, All);
 #define LOCTEXT_NAMESPACE "SkeletalMeshImport"
 
 /** Check that root bone is the same, and that any bones that are common have the correct parent. */
-bool SkeletonsAreCompatible( const FReferenceSkeleton& NewSkel, const FReferenceSkeleton& ExistSkel, bool bFailNoError)
+bool SkeletonsAreCompatible(const FReferenceSkeleton& NewSkel, const FReferenceSkeleton& ExistSkel, bool bFailNoError)
 {
-	if(NewSkel.GetBoneName(0) != ExistSkel.GetBoneName(0))
+	if (NewSkel.GetBoneName(0) != ExistSkel.GetBoneName(0))
 	{
 		if (!bFailNoError)
 		{
@@ -51,7 +51,7 @@ bool SkeletonsAreCompatible( const FReferenceSkeleton& NewSkel, const FReference
 		return false;
 	}
 
-	for(int32 i=1; i<NewSkel.GetRawBoneNum(); i++)
+	for (int32 i = 1; i < NewSkel.GetRawBoneNum(); i++)
 	{
 		// See if bone is in both skeletons.
 		int32 NewBoneIndex = i;
@@ -59,12 +59,12 @@ bool SkeletonsAreCompatible( const FReferenceSkeleton& NewSkel, const FReference
 		int32 BBoneIndex = ExistSkel.FindBoneIndex(NewBoneName);
 
 		// If it is, check parents are the same.
-		if(BBoneIndex != INDEX_NONE)
+		if (BBoneIndex != INDEX_NONE)
 		{
-			FName NewParentName = NewSkel.GetBoneName( NewSkel.GetParentIndex(NewBoneIndex) );
-			FName ExistParentName = ExistSkel.GetBoneName( ExistSkel.GetParentIndex(BBoneIndex) );
+			FName NewParentName = NewSkel.GetBoneName(NewSkel.GetParentIndex(NewBoneIndex));
+			FName ExistParentName = ExistSkel.GetBoneName(ExistSkel.GetParentIndex(BBoneIndex));
 
-			if(NewParentName != ExistParentName)
+			if (NewParentName != ExistParentName)
 			{
 				if (!bFailNoError)
 				{
@@ -82,11 +82,11 @@ bool SkeletonsAreCompatible( const FReferenceSkeleton& NewSkel, const FReference
 
 /**
 * Process and fill in the mesh Materials using the raw binary import data
-* 
+*
 * @param Materials - [out] array of materials to update
 * @param ImportData - raw binary import data to process
 */
-void ProcessImportMeshMaterials(TArray<FSkeletalMaterial>& Materials, FSkeletalMeshImportData& ImportData )
+void ProcessImportMeshMaterials(TArray<FSkeletalMaterial>& Materials, FSkeletalMeshImportData& ImportData)
 {
 	TArray <SkeletalMeshImportData::FMaterial>&	ImportedMaterials = ImportData.Materials;
 
@@ -94,13 +94,13 @@ void ProcessImportMeshMaterials(TArray<FSkeletalMaterial>& Materials, FSkeletalM
 	// material name, cut off anything in front of the dot (beyond are special flags).
 	Materials.Empty();
 	int32 SkinOffset = INDEX_NONE;
-	for( int32 MatIndex=0; MatIndex < ImportedMaterials.Num(); ++MatIndex)
-	{			
+	for (int32 MatIndex = 0; MatIndex < ImportedMaterials.Num(); ++MatIndex)
+	{
 		const SkeletalMeshImportData::FMaterial& ImportedMaterial = ImportedMaterials[MatIndex];
 
 		UMaterialInterface* Material = NULL;
 		FString MaterialNameNoSkin = ImportedMaterial.MaterialImportName;
-		if( ImportedMaterial.Material.IsValid() )
+		if (ImportedMaterial.Material.IsValid())
 		{
 			Material = ImportedMaterial.Material.Get();
 		}
@@ -125,21 +125,21 @@ void ProcessImportMeshMaterials(TArray<FSkeletalMaterial>& Materials, FSkeletalM
 		}
 
 		const bool bEnableShadowCasting = true;
-		Materials.Add( FSkeletalMaterial( Material, bEnableShadowCasting, false, Material != nullptr ? Material->GetFName() : FName(*MaterialNameNoSkin), FName(*(ImportedMaterial.MaterialImportName)) ) );
+		Materials.Add(FSkeletalMaterial(Material, bEnableShadowCasting, false, Material != nullptr ? Material->GetFName() : FName(*MaterialNameNoSkin), FName(*(ImportedMaterial.MaterialImportName))));
 	}
 
-	int32 NumMaterialsToAdd = FMath::Max<int32>( ImportedMaterials.Num(), ImportData.MaxMaterialIndex + 1 );
+	int32 NumMaterialsToAdd = FMath::Max<int32>(ImportedMaterials.Num(), ImportData.MaxMaterialIndex + 1);
 
 	// Pad the material pointers
-	while( NumMaterialsToAdd > Materials.Num() )
+	while (NumMaterialsToAdd > Materials.Num())
 	{
-		Materials.Add( FSkeletalMaterial( NULL, true, false, NAME_None, NAME_None ) );
+		Materials.Add(FSkeletalMaterial(NULL, true, false, NAME_None, NAME_None));
 	}
 }
 
 /**
 * Process and fill in the mesh ref skeleton bone hierarchy using the raw binary import data
-* 
+*
 * @param RefSkeleton - [out] reference skeleton hierarchy to update
 * @param SkeletalDepth - [out] depth of the reference skeleton hierarchy
 * @param ImportData - raw binary import data to process
@@ -155,14 +155,14 @@ bool ProcessImportMeshSkeleton(const USkeleton* SkeletonAsset, FReferenceSkeleto
 	FReferenceSkeletonModifier RefSkelModifier(RefSkeleton, SkeletonAsset);
 
 	// Digest bones to the serializable format.
-	for( int32 b=0; b<RefBonesBinary.Num(); b++ )
+	for (int32 b = 0; b < RefBonesBinary.Num(); b++)
 	{
-		const SkeletalMeshImportData::FBone & BinaryBone = RefBonesBinary[ b ];
-		const FString BoneName = FSkeletalMeshImportData::FixupBoneName( BinaryBone.Name );
+		const SkeletalMeshImportData::FBone & BinaryBone = RefBonesBinary[b];
+		const FString BoneName = FSkeletalMeshImportData::FixupBoneName(BinaryBone.Name);
 		const FMeshBoneInfo BoneInfo(FName(*BoneName, FNAME_Add), BinaryBone.Name, BinaryBone.ParentIndex);
 		const FTransform BoneTransform(BinaryBone.BonePos.Transform);
 
-		if(RefSkeleton.FindRawBoneIndex(BoneInfo.Name) != INDEX_NONE)
+		if (RefSkeleton.FindRawBoneIndex(BoneInfo.Name) != INDEX_NONE)
 		{
 			UnFbx::FFbxImporter* FFbxImporter = UnFbx::FFbxImporter::GetInstance();
 			FFbxImporter->AddTokenizedErrorMessage(FTokenizedMessage::Create(EMessageSeverity::Error, FText::Format(LOCTEXT("SkeletonHasDuplicateBones", "Skeleton has non-unique bone names.\nBone named '{0}' encountered more than once."), FText::FromName(BoneInfo.Name))), FFbxErrors::SkeletalMesh_DuplicateBones);
@@ -176,19 +176,19 @@ bool ProcessImportMeshSkeleton(const USkeleton* SkeletonAsset, FReferenceSkeleto
 	SkeletalDepth = 0;
 
 	TArray<int32> SkeletalDepths;
-	SkeletalDepths.Empty( RefBonesBinary.Num() );
-	SkeletalDepths.AddZeroed( RefBonesBinary.Num() );
-	for( int32 b=0; b < RefSkeleton.GetRawBoneNum(); b++ )
+	SkeletalDepths.Empty(RefBonesBinary.Num());
+	SkeletalDepths.AddZeroed(RefBonesBinary.Num());
+	for (int32 b = 0; b < RefSkeleton.GetRawBoneNum(); b++)
 	{
-		int32 Parent	= RefSkeleton.GetRawParentIndex(b);
-		int32 Depth	= 1.0f;
+		int32 Parent = RefSkeleton.GetRawParentIndex(b);
+		int32 Depth = 1.0f;
 
-		SkeletalDepths[b]	= 1.0f;
-		if( Parent != INDEX_NONE )
+		SkeletalDepths[b] = 1.0f;
+		if (Parent != INDEX_NONE)
 		{
 			Depth += SkeletalDepths[Parent];
 		}
-		if( SkeletalDepth < Depth )
+		if (SkeletalDepth < Depth)
 		{
 			SkeletalDepth = Depth;
 		}
@@ -332,11 +332,11 @@ ExistingSkelMeshData* SaveExistingSkelMeshData(USkeletalMesh* ExistingSkelMesh,
 	}
 	ExistingMeshDataPtr->ExistingRetargetBasePose = ExistingSkelMesh->RetargetBasePose;
 
-	if( ImportedResource->LODModels.Num() > 0 &&
-		ExistingSkelMesh->GetLODNum() == ImportedResource->LODModels.Num() )
+	if (ImportedResource->LODModels.Num() > 0 &&
+		ExistingSkelMesh->GetLODNum() == ImportedResource->LODModels.Num())
 	{
 		int32 OffsetReductionLODIndex = 0;
-		FSkeletalMeshLODInfo* LODInfo = ExistingSkelMesh->GetLODInfo( SafeReimportLODIndex);
+		FSkeletalMeshLODInfo* LODInfo = ExistingSkelMesh->GetLODInfo(SafeReimportLODIndex);
 		ExistingMeshDataPtr->bIsReimportLODReduced = (LODInfo && LODInfo->bHasBeenSimplified);
 		if (ExistingMeshDataPtr->bIsReimportLODReduced)
 		{
@@ -354,7 +354,7 @@ ExistingSkelMeshData* SaveExistingSkelMeshData(USkeletalMesh* ExistingSkelMesh,
 		}
 
 		// Copy off the remaining LODs.
-		for ( int32 LODModelIndex = 0 ; LODModelIndex < ImportedResource->LODModels.Num() ; ++LODModelIndex )
+		for (int32 LODModelIndex = 0; LODModelIndex < ImportedResource->LODModels.Num(); ++LODModelIndex)
 		{
 			FSkeletalMeshLODModel& LODModel = ImportedResource->LODModels[LODModelIndex];
 			LODModel.RawPointIndices.Lock( LOCK_READ_ONLY );
@@ -388,18 +388,18 @@ ExistingSkelMeshData* SaveExistingSkelMeshData(USkeletalMesh* ExistingSkelMesh,
 
 		ExistingMeshDataPtr->ExistingLODInfo = ExistingSkelMesh->GetLODInfoArray();
 		ExistingMeshDataPtr->ExistingRefSkeleton = ExistingSkelMesh->RefSkeleton;
-		
+
 	}
 
 	// First asset should be the one that the skeletal mesh should point too
 	ExistingMeshDataPtr->ExistingPhysicsAssets.Empty();
-	ExistingMeshDataPtr->ExistingPhysicsAssets.Add( ExistingSkelMesh->PhysicsAsset );
+	ExistingMeshDataPtr->ExistingPhysicsAssets.Add(ExistingSkelMesh->PhysicsAsset);
 	for (TObjectIterator<UPhysicsAsset> It; It; ++It)
 	{
 		UPhysicsAsset* PhysicsAsset = *It;
-		if ( PhysicsAsset->PreviewSkeletalMesh == ExistingSkelMesh && ExistingSkelMesh->PhysicsAsset != PhysicsAsset )
+		if (PhysicsAsset->PreviewSkeletalMesh == ExistingSkelMesh && ExistingSkelMesh->PhysicsAsset != PhysicsAsset)
 		{
-			ExistingMeshDataPtr->ExistingPhysicsAssets.Add( PhysicsAsset );
+			ExistingMeshDataPtr->ExistingPhysicsAssets.Add(PhysicsAsset);
 		}
 	}
 
@@ -415,7 +415,7 @@ ExistingSkelMeshData* SaveExistingSkelMeshData(USkeletalMesh* ExistingSkelMesh,
 
 	ExistingMeshDataPtr->ExistingMorphTargets.Empty(ExistingSkelMesh->MorphTargets.Num());
 	ExistingMeshDataPtr->ExistingMorphTargets.Append(ExistingSkelMesh->MorphTargets);
-	
+
 	ExistingMeshDataPtr->ExistingAssetImportData = ExistingSkelMesh->AssetImportData;
 	ExistingMeshDataPtr->ExistingThumbnailInfo = ExistingSkelMesh->ThumbnailInfo;
 
@@ -645,7 +645,7 @@ namespace SkeletalMeshHelper
 				{
 					if (SrcVertex.InfluenceWeights[InfluenceIndex] > 0.0f)
 					{
-						Section.MaxBoneInfluences = FMath::Max(Section.MaxBoneInfluences, InfluenceIndex+1);
+						Section.MaxBoneInfluences = FMath::Max(Section.MaxBoneInfluences, InfluenceIndex + 1);
 						//Copy the weight
 						DestVertex.InfluenceWeights[InfluenceIndex] = SrcVertex.InfluenceWeights[InfluenceIndex];
 						//Copy the bone ID
@@ -705,14 +705,14 @@ void RestoreExistingSkelMeshData(ExistingSkelMeshData* MeshData, USkeletalMesh*
 	RemapMaterial.AddZeroed(SkeletalMesh->Materials.Num());
 	TArray<FName> RemapMaterialName;
 	RemapMaterialName.AddZeroed(SkeletalMesh->Materials.Num());
-	
+
 	bool bMaterialReset = false;
 	if (MeshData->bSaveRestoreMaterials)
 	{
 		UnFbx::EFBXReimportDialogReturnOption ReturnOption;
 		//Ask the user to match the materials conflict
 		UnFbx::FFbxImporter::PrepareAndShowMaterialConflictDialog<FSkeletalMaterial>(MeshData->ExistingMaterials, SkeletalMesh->Materials, RemapMaterial, RemapMaterialName, bCanShowDialog, false, ReturnOption);
-		
+
 		if (ReturnOption != UnFbx::EFBXReimportDialogReturnOption::FBXRDRO_ResetToFbx)
 		{
 			//Build a ordered material list that try to keep intact the existing material list
@@ -952,7 +952,7 @@ void RestoreExistingSkelMeshData(ExistingSkelMeshData* MeshData, USkeletalMesh*
 
 		SkeletalMesh->Skeleton = MeshData->ExistingSkeleton;
 		SkeletalMesh->PostProcessAnimBlueprint = MeshData->ExistingPostProcessAnimBlueprint;
-		
+
 		// Copy mirror table.
 		SkeletalMesh->ImportMirrorTable(MeshData->ExistingMirrorTable);
 
@@ -967,7 +967,7 @@ void RestoreExistingSkelMeshData(ExistingSkelMeshData* MeshData, USkeletalMesh*
 
 		for (UClothingAssetBase* ClothingAsset : SkeletalMesh->MeshClothingAssets)
 		{
-			if(ClothingAsset)
+			if (ClothingAsset)
 				ClothingAsset->RefreshBoneMapping(SkeletalMesh);
 		}
 
@@ -981,7 +981,7 @@ void RestoreExistingSkelMeshData(ExistingSkelMeshData* MeshData, USkeletalMesh*
 		if (SafeReimportLODIndex < 1 && SkeletalMesh->GetLODInfoArray().IsValidIndex(SafeReimportLODIndex))
 		{
 			FSkeletalMeshLODInfo& BaseLODInfo = SkeletalMesh->GetLODInfoArray()[SafeReimportLODIndex];
-			if(bMaterialReset)
+			if (bMaterialReset)
 			{
 				//If we reset the material array there is no point keeping the user changes
 				BaseLODInfo.LODMaterialMap.Empty();
@@ -1045,37 +1045,42 @@ void RestoreExistingSkelMeshData(ExistingSkelMeshData* MeshData, USkeletalMesh*
 			{
 				break;
 			}
-			int32 CurrentSectionMaterialIndex = NewSkelMeshLodModel.Sections[SectionIndex].MaterialIndex;
-			FName CurrentSectionImportedMaterialName = SkeletalMesh->Materials[CurrentSectionMaterialIndex].ImportedMaterialSlotName;
-			for (int32 ExistSectionIndex = 0; ExistSectionIndex < MeshData->LastImportMeshLodSectionMaterialData[SafeReimportLODIndex].Num(); ++ExistSectionIndex)
+
+			// If the new mesh has different number of sections than last import, then we shouldn't remap section materials.
+			if (NewSkelMeshLodModel.Sections.Num() == MeshData->LastImportMeshLodSectionMaterialData[SafeReimportLODIndex].Num())
 			{
-				if (!MeshData->LastImportMeshLodSectionMaterialData[SafeReimportLODIndex].IsValidIndex(ExistSectionIndex) || !MeshData->ExistingImportMeshLodSectionMaterialData[SafeReimportLODIndex].IsValidIndex(ExistSectionIndex))
-				{
-					continue;
-				}
-				//The last import slot name is use to match the New import slot name
-				//If the user has change the slot in the editor the ExistMeshSectionSlotName will be different.
-				//This is why we use the ExistMeshSectionSlotName to restore the data and use the LastImportedMeshSectionSlotName to rematch the section
-				FName LastImportedMeshSectionSlotName = MeshData->LastImportMeshLodSectionMaterialData[SafeReimportLODIndex][ExistSectionIndex];
-				if (LastImportedMeshSectionSlotName != CurrentSectionImportedMaterialName)
+				int32 CurrentSectionMaterialIndex = NewSkelMeshLodModel.Sections[SectionIndex].MaterialIndex;
+				FName CurrentSectionImportedMaterialName = SkeletalMesh->Materials[CurrentSectionMaterialIndex].ImportedMaterialSlotName;
+				for (int32 ExistSectionIndex = 0; ExistSectionIndex < MeshData->LastImportMeshLodSectionMaterialData[SafeReimportLODIndex].Num(); ++ExistSectionIndex)
 				{
-					//This material do not match
-					continue;
-				}
-				//Restore the material slot
-				FName ExistMeshSectionSlotName = MeshData->ExistingImportMeshLodSectionMaterialData[SafeReimportLODIndex][ExistSectionIndex].ImportedMaterialSlotName;
+					if (!MeshData->LastImportMeshLodSectionMaterialData[SafeReimportLODIndex].IsValidIndex(ExistSectionIndex) || !MeshData->ExistingImportMeshLodSectionMaterialData[SafeReimportLODIndex].IsValidIndex(ExistSectionIndex))
+					{
+						continue;
+					}
+					//The last import slot name is use to match the New import slot name
+					//If the user has change the slot in the editor the ExistMeshSectionSlotName will be different.
+					//This is why we use the ExistMeshSectionSlotName to restore the data and use the LastImportedMeshSectionSlotName to rematch the section
+					FName LastImportedMeshSectionSlotName = MeshData->LastImportMeshLodSectionMaterialData[SafeReimportLODIndex][ExistSectionIndex];
+					if (LastImportedMeshSectionSlotName != CurrentSectionImportedMaterialName)
+					{
+						//This material do not match
+						continue;
+					}
+					//Restore the material slot
+					FName ExistMeshSectionSlotName = MeshData->ExistingImportMeshLodSectionMaterialData[SafeReimportLODIndex][ExistSectionIndex].ImportedMaterialSlotName;
 
-				//Override the new section material index to use the one that the user set
-				for (int32 RemapMaterialIndex = 0; RemapMaterialIndex < SkeletalMesh->Materials.Num(); ++RemapMaterialIndex)
-				{
-					const FSkeletalMaterial &NewSectionMaterial = SkeletalMesh->Materials[RemapMaterialIndex];
-					if (NewSectionMaterial.ImportedMaterialSlotName == ExistMeshSectionSlotName)
+					//Override the new section material index to use the one that the user set
+					for (int32 RemapMaterialIndex = 0; RemapMaterialIndex < SkeletalMesh->Materials.Num(); ++RemapMaterialIndex)
 					{
-						NewSkelMeshLodModel.Sections[SectionIndex].MaterialIndex = RemapMaterialIndex;
-						break;
+						const FSkeletalMaterial &NewSectionMaterial = SkeletalMesh->Materials[RemapMaterialIndex];
+						if (NewSectionMaterial.ImportedMaterialSlotName == ExistMeshSectionSlotName)
+						{
+							NewSkelMeshLodModel.Sections[SectionIndex].MaterialIndex = RemapMaterialIndex;
+							break;
+						}
 					}
+					break;
 				}
-				break;
 			}
 		}
 		//Make sure we reset the User section array to only what we have in the fbx
diff --git a/Engine/Source/Editor/UnrealEd/Private/Fbx/FbxMainExport.cpp b/Engine/Source/Editor/UnrealEd/Private/Fbx/FbxMainExport.cpp
--- a/Engine/Source/Editor/UnrealEd/Private/Fbx/FbxMainExport.cpp
+++ b/Engine/Source/Editor/UnrealEd/Private/Fbx/FbxMainExport.cpp
@@ -1427,25 +1427,24 @@ FString FFbxExporter::FMatineeNodeNameAdapter::GetActorNodeName(const AActor* Ac
 	return NodeName;
 }
 
-
-FFbxExporter::FMatineeAnimTrackAdapter::FMatineeAnimTrackAdapter( AMatineeActor* InMatineeActor )
+FMatineeAnimTrackAdapter::FMatineeAnimTrackAdapter( AMatineeActor* InMatineeActor )
 {
 	MatineeActor = InMatineeActor;
 }
 
-int32 FFbxExporter::FMatineeAnimTrackAdapter::GetLocalStartFrame() const
+int32 FMatineeAnimTrackAdapter::GetLocalStartFrame() const
 {
 	return 0;
 }
 
 
-int32 FFbxExporter::FMatineeAnimTrackAdapter::GetLength() const
+int32 FMatineeAnimTrackAdapter::GetLength() const
 {
 	return MatineeActor->MatineeData->InterpLength * (1.f / GetFrameRate());
 }
 
 
-void FFbxExporter::FMatineeAnimTrackAdapter::UpdateAnimation( int32 LocalFrame )
+void FMatineeAnimTrackAdapter::UpdateAnimation( int32 LocalFrame )
 {
 	float Time = LocalFrame * GetFrameRate();
 	MatineeActor->UpdateInterp( Time, true );
@@ -1580,29 +1579,28 @@ FbxNode* FFbxExporter::FLevelSequenceNodeNameAdapter::GetFbxNode(UObject* InObje
 	return nullptr;
 }
 
-FFbxExporter::FLevelSequenceAnimTrackAdapter::FLevelSequenceAnimTrackAdapter( IMovieScenePlayer* InMovieScenePlayer, UMovieScene* InMovieScene, const FMovieSceneSequenceTransform& InRootToLocalTransform)
+FLevelSequenceAnimTrackAdapter::FLevelSequenceAnimTrackAdapter( IMovieScenePlayer* InMovieScenePlayer, UMovieScene* InMovieScene, const FMovieSceneSequenceTransform& InRootToLocalTransform)
 {
 	MovieScenePlayer = InMovieScenePlayer;
 	MovieScene = InMovieScene;
 	RootToLocalTransform = InRootToLocalTransform;
 }
 
-
-int32 FFbxExporter::FLevelSequenceAnimTrackAdapter::GetLocalStartFrame() const
+int32 FLevelSequenceAnimTrackAdapter::GetLocalStartFrame() const
 {
 	FFrameRate TickResolution = MovieScene->GetTickResolution();
 	FFrameRate DisplayRate = MovieScene->GetDisplayRate();
 	return FFrameRate::TransformTime(FFrameTime(MovieScene::DiscreteInclusiveLower(MovieScene->GetPlaybackRange())), TickResolution, DisplayRate).RoundToFrame().Value;
 }
 
-int32 FFbxExporter::FLevelSequenceAnimTrackAdapter::GetStartFrame() const
+int32 FLevelSequenceAnimTrackAdapter::GetStartFrame() const
 {
 	FFrameRate TickResolution = MovieScene->GetTickResolution();
 	FFrameRate DisplayRate = MovieScene->GetDisplayRate();
 	return FFrameRate::TransformTime(FFrameTime(MovieScene::DiscreteInclusiveLower(MovieScene->GetPlaybackRange()) * RootToLocalTransform.InverseLinearOnly()), TickResolution, DisplayRate).RoundToFrame().Value;
 }
 
-int32 FFbxExporter::FLevelSequenceAnimTrackAdapter::GetLength() const
+int32 FLevelSequenceAnimTrackAdapter::GetLength() const
 {
 	FFrameRate TickResolution = MovieScene->GetTickResolution();
 	FFrameRate DisplayRate = MovieScene->GetDisplayRate();
@@ -1610,7 +1608,7 @@ int32 FFbxExporter::FLevelSequenceAnimTrackAdapter::GetLength() const
 }
 
 
-void FFbxExporter::FLevelSequenceAnimTrackAdapter::UpdateAnimation( int32 LocalFrame )
+void FLevelSequenceAnimTrackAdapter::UpdateAnimation( int32 LocalFrame )
 {	
 	FFrameRate TickResolution = MovieScene->GetTickResolution();
 	FFrameRate DisplayRate = MovieScene->GetDisplayRate();
@@ -1623,7 +1621,7 @@ void FFbxExporter::FLevelSequenceAnimTrackAdapter::UpdateAnimation( int32 LocalF
 	MovieScenePlayer->GetEvaluationTemplate().Evaluate( Context, *MovieScenePlayer );
 }
 
-float FFbxExporter::FLevelSequenceAnimTrackAdapter::GetFrameRate() const
+float FLevelSequenceAnimTrackAdapter::GetFrameRate() const
 {
 	return MovieScene->GetDisplayRate().AsDecimal();
 }
diff --git a/Engine/Source/Editor/UnrealEd/Private/FbxExporter.h b/Engine/Source/Editor/UnrealEd/Private/FbxExporter.h
--- a/Engine/Source/Editor/UnrealEd/Private/FbxExporter.h
+++ b/Engine/Source/Editor/UnrealEd/Private/FbxExporter.h
@@ -48,7 +48,47 @@
 
 namespace UnFbx
 {
+	/** Adapter interface which allows ExportAnimTrack to act on both sequencer and matinee data. */
+	class UNREALED_API IAnimTrackAdapter
+	{
+	public:
+		virtual ~IAnimTrackAdapter() {};
+		virtual int32 GetLocalStartFrame() const = 0;
+		virtual int32 GetStartFrame() const { return GetLocalStartFrame(); }
+		virtual int32 GetLength() const = 0;
+		/** Updates the runtime state of the animation track to the specified frame. */
+		virtual void UpdateAnimation(int32 LocalFrame) = 0;
+		virtual float GetFrameRate() const { return 1.f / DEFAULT_SAMPLERATE; }
+	};
 
+	/** An anim track adapter for matinee. */
+	class UNREALED_API FMatineeAnimTrackAdapter : public IAnimTrackAdapter
+	{
+	public:
+		FMatineeAnimTrackAdapter(AMatineeActor* InMatineeActor);
+		virtual int32 GetLocalStartFrame() const override;
+		virtual int32 GetLength() const override;
+		virtual void UpdateAnimation(int32 LocalFrame) override;
+
+	private:
+		AMatineeActor* MatineeActor;
+	};
+	/** An anim track adapter for a level sequence. */
+	class UNREALED_API FLevelSequenceAnimTrackAdapter : public IAnimTrackAdapter
+	{
+	public:
+		FLevelSequenceAnimTrackAdapter(IMovieScenePlayer* InMovieScenePlayer, UMovieScene* InMovieScene, const FMovieSceneSequenceTransform& InRootToLocalTransform);
+		virtual int32 GetLocalStartFrame() const override;
+		virtual int32 GetStartFrame() const override;
+		virtual int32 GetLength() const override;
+		virtual void UpdateAnimation(int32 LocalFrame) override;
+		virtual float GetFrameRate() const override;
+
+	private:
+		IMovieScenePlayer* MovieScenePlayer;
+		UMovieScene* MovieScene;
+		FMovieSceneSequenceTransform RootToLocalTransform;
+	};
 /**
  * Main FBX Exporter class.
  */
@@ -266,49 +306,6 @@
 	UFbxExportOption *ExportOptionsUI;
 	UFbxExportOption *ExportOptionsOverride;
 
-	/** Adapter interface which allows ExportAnimTrack to act on both sequencer and matinee data. */
-	class IAnimTrackAdapter
-	{
-	public:
-		virtual ~IAnimTrackAdapter() {}
-		virtual int32 GetLocalStartFrame() const = 0;
-		virtual int32 GetStartFrame() const { return GetLocalStartFrame(); }
-		virtual int32 GetLength() const = 0;
-		/** Updates the runtime state of the animation track to the specified frame. */
-		virtual void UpdateAnimation( int32 LocalFrame) = 0;
-		virtual float GetFrameRate() const { return 1.f / DEFAULT_SAMPLERATE; }
-	};
-
-	/** An anim track adapter for matinee. */
-	class FMatineeAnimTrackAdapter : public IAnimTrackAdapter
-	{
-	public:
-		FMatineeAnimTrackAdapter( AMatineeActor* InMatineeActor );
-		virtual int32 GetLocalStartFrame() const override;
-		virtual int32 GetLength() const override;
-		virtual void UpdateAnimation(int32 LocalFrame) override;
-
-	private:
-		AMatineeActor* MatineeActor;
-	};
-
-	/** An anim track adapter for a level sequence. */
-	class FLevelSequenceAnimTrackAdapter : public FFbxExporter::IAnimTrackAdapter
-	{
-	public:
-		FLevelSequenceAnimTrackAdapter( IMovieScenePlayer* InMovieScenePlayer, UMovieScene* InMovieScene, const FMovieSceneSequenceTransform& InRootToLocalTransform );
-		virtual int32 GetLocalStartFrame() const override;
-		virtual int32 GetStartFrame() const override;
-		virtual int32 GetLength() const override;
-		virtual void UpdateAnimation(int32 LocalFrame) override;
-		virtual float GetFrameRate() const override;
-
-	private:
-		IMovieScenePlayer* MovieScenePlayer;
-		UMovieScene* MovieScene;
-		FMovieSceneSequenceTransform RootToLocalTransform;
-	};
-
 	/**
 	* Export Anim Track of the given SkeletalMeshComponent
 	*/
diff --git a/Engine/Source/Editor/UnrealEd/Private/PreferenceStubs.cpp b/Engine/Source/Editor/UnrealEd/Private/PreferenceStubs.cpp
--- a/Engine/Source/Editor/UnrealEd/Private/PreferenceStubs.cpp
+++ b/Engine/Source/Editor/UnrealEd/Private/PreferenceStubs.cpp
@@ -10,6 +10,7 @@
 #include "Preferences/AnimationBlueprintEditorOptions.h"
 #include "Preferences/PhysicsAssetEditorOptions.h"
 #include "Preferences/MaterialStatsOptions.h"
+#include "FrameNumberDisplayFormat.h"
 
 // @todo find a better place for all of this, preferably in the appropriate modules
 // though this would require the classes to be relocated as well
@@ -123,8 +124,8 @@ UPersonaOptions::UPersonaOptions(const FObjectInitializer& ObjectInitializer)
 		EditorOptions.SetViewportConfigsToDefault();
 	}
 
-	SectionTimingNodeColor = FLinearColor(0.0f, 1.0f, 0.0f);
-	NotifyTimingNodeColor = FLinearColor(1.0f, 0.0f, 0.0f);
+	SectionTimingNodeColor = FLinearColor(0.39f, 0.39f, 1.0f, 0.75f);
+	NotifyTimingNodeColor = FLinearColor(0.8f, 0.1f, 0.1f);
 	BranchingPointTimingNodeColor = FLinearColor(0.5f, 1.0f, 1.0f);
 
 	bAutoAlignFloorToMesh = true;
@@ -132,6 +133,18 @@ UPersonaOptions::UPersonaOptions(const FObjectInitializer& ObjectInitializer)
 	NumFolderFiltersInAssetBrowser = 2;
 
 	bUseAudioAttenuation = true;
+
+	CurveEditorSnapInterval = 0.01f;
+
+	// Default to millisecond resolution
+	TimelineScrubSnapValue = 1000;
+
+	TimelineDisplayFormat = EFrameNumberDisplayFormats::Frames;
+
+	bTimelineDisplayPercentage = true;
+	bTimelineDisplayFormatSecondary = true;
+
+	bTimelineDisplayCurveKeys = false;
 }
 
 void UPersonaOptions::SetShowGrid( bool bInShowGrid )
diff --git a/Engine/Source/Editor/UnrealEd/Private/PreviewMaterial.cpp b/Engine/Source/Editor/UnrealEd/Private/PreviewMaterial.cpp
--- a/Engine/Source/Editor/UnrealEd/Private/PreviewMaterial.cpp
+++ b/Engine/Source/Editor/UnrealEd/Private/PreviewMaterial.cpp
@@ -107,8 +107,9 @@ public:
 				bool bSkinCache = IsGPUSkinCacheAvailable(Platform) && (VertexFactoryType == FindVertexFactoryType(FName(TEXT("FGPUSkinPassthroughVertexFactory"), FNAME_Find)));
 					
 				if (
-					VertexFactoryType != FindVertexFactoryType(FName(TEXT("TGPUSkinVertexFactoryfalse"), FNAME_Find)) &&
-					VertexFactoryType != FindVertexFactoryType(FName(TEXT("TGPUSkinVertexFactorytrue"), FNAME_Find)) &&
+					VertexFactoryType != FindVertexFactoryType(FName(TEXT("TGPUSkinVertexFactoryDefault"), FNAME_Find)) &&
+					VertexFactoryType != FindVertexFactoryType(FName(TEXT("TGPUSkinVertexFactoryExtra"), FNAME_Find)) &&
+					VertexFactoryType != FindVertexFactoryType(FName(TEXT("TGPUSkinVertexFactoryUnlimited"), FNAME_Find)) &&
 					!bSkinCache
 					)
 				{
diff --git a/Engine/Source/Editor/UnrealEd/Private/ThumbnailRendering/SkeletalMeshThumbnailRenderer.cpp b/Engine/Source/Editor/UnrealEd/Private/ThumbnailRendering/SkeletalMeshThumbnailRenderer.cpp
--- a/Engine/Source/Editor/UnrealEd/Private/ThumbnailRendering/SkeletalMeshThumbnailRenderer.cpp
+++ b/Engine/Source/Editor/UnrealEd/Private/ThumbnailRendering/SkeletalMeshThumbnailRenderer.cpp
@@ -24,6 +24,8 @@ void USkeletalMeshThumbnailRenderer::Draw(UObject* Object, int32 X, int32 Y, uin
 		}
 
 		ThumbnailScene->SetSkeletalMesh(SkeletalMesh);
+		AddAdditionalPreviewSceneContent(Object, ThumbnailScene->GetWorld());
+
 		FSceneViewFamilyContext ViewFamily( FSceneViewFamily::ConstructionValues( RenderTarget, ThumbnailScene->GetScene(), FEngineShowFlags(ESFIM_Game) )
 			.SetWorldTimes(FApp::GetCurrentTime() - GStartTime, FApp::GetDeltaTime(), FApp::GetCurrentTime() - GStartTime));
 
diff --git a/Engine/Source/Editor/UnrealEd/Public/AnimationEditorUtils.h b/Engine/Source/Editor/UnrealEd/Public/AnimationEditorUtils.h
--- a/Engine/Source/Editor/UnrealEd/Public/AnimationEditorUtils.h
+++ b/Engine/Source/Editor/UnrealEd/Public/AnimationEditorUtils.h
@@ -19,7 +19,7 @@
 
 class FMenuBuilder;
 class UAnimBlueprint;
-class UAnimCompress;
+class UAnimBoneCompressionSettings;
 class UAnimSequence;
 class UEdGraph;
 class UPoseWatch;
@@ -69,18 +69,80 @@ protected:
 	static FText LastUsedAssetPath;
 };
 
+/** A struct containing the settings to control the SAnimationCompressionSelectionDialog creation. */
+struct FAnimationCompressionSelectionDialogConfig
+{
+	FText DialogTitleOverride;
+	FVector2D WindowSizeOverride;
+
+	UAnimBoneCompressionSettings* DefaultSelectedAsset;
+
+	FAnimationCompressionSelectionDialogConfig()
+		: WindowSizeOverride(ForceInitToZero)
+		, DefaultSelectedAsset(nullptr)
+	{}
+};
+
+/** Dialog to prompt user to select an animation compression settings asset. */
+class UNREALED_API SAnimationCompressionSelectionDialog : public SCompoundWidget
+{
+public:
+	/** Called from the Dialog when an asset has been selected. */
+	DECLARE_DELEGATE_OneParam(FOnAssetSelected, const FAssetData& /*SelectedAsset*/);
+
+	SLATE_BEGIN_ARGS(SAnimationCompressionSelectionDialog) {}
+
+	SLATE_END_ARGS()
+
+	SAnimationCompressionSelectionDialog();
+	virtual ~SAnimationCompressionSelectionDialog();
+
+	virtual void Construct(const FArguments& InArgs, const FAnimationCompressionSelectionDialogConfig& InConfig);
+
+	/** Sets the delegate handler for when an open operation is committed */
+	void SetOnAssetSelected(const FOnAssetSelected& InHandler);
+
+private:
+	void DoSelectAsset(const FAssetData& SelectedAsset);
+	FReply OnConfirmClicked();
+	FReply OnCancelClicked();
+	void CloseDialog();
+	void OnAssetSelected(const FAssetData& AssetData);
+	void OnAssetsActivated(const TArray<FAssetData>& SelectedAssets, EAssetTypeActivationMethod::Type ActivationType);
+	bool IsConfirmButtonEnabled() const;
+
+	/** Asset Picker used by the dialog */
+	TSharedPtr<SWidget> AssetPicker;
+
+	/** The assets that are currently selected in the asset picker */
+	TArray<FAssetData> CurrentlySelectedAssets;
+
+	/** Used to specify that a valid asset was chosen */
+	bool bValidAssetChosen;
+
+	/** Fired when assets are chosen for open. Only fired in open dialogs. */
+	FOnAssetSelected OnAssetSelectedHandler;
+
+	/** Used to get the currently selected assets */
+	FGetCurrentSelectionDelegate GetCurrentSelectionDelegate;
+};
+
 /** Defines FCanExecuteAction delegate interface. Returns false to force the caller to delete the just created assets*/
 DECLARE_DELEGATE_RetVal_OneParam(bool, FAnimAssetCreated, TArray<class UObject*>);
 
 //Animation editor utility functions
 namespace AnimationEditorUtils
 {
+	UNREALED_API FAssetData CreateModalAnimationCompressionSelectionDialog(const FAnimationCompressionSelectionDialogConfig& InConfig);
+
 	UNREALED_API void CreateAnimationAssets(const TArray<TWeakObjectPtr<UObject>>& SkeletonsOrSkeletalMeshes, TSubclassOf<UAnimationAsset> AssetClass, const FString& InPrefix, FAnimAssetCreated AssetCreated, UObject* NameBaseObject = nullptr, bool bDoNotShowNameDialog = false);
 	
 	UNREALED_API void CreateNewAnimBlueprint(TArray<TWeakObjectPtr<UObject>> SkeletonsOrSkeletalMeshes, FAnimAssetCreated AssetCreated, bool bInContentBrowser);
 	UNREALED_API void FillCreateAssetMenu(FMenuBuilder& MenuBuilder, const TArray<TWeakObjectPtr<UObject>>& SkeletonsOrSkeletalMeshes, FAnimAssetCreated AssetCreated, bool bInContentBrowser=true);
 	UNREALED_API void CreateUniqueAssetName(const FString& InBasePackageName, const FString& InSuffix, FString& OutPackageName, FString& OutAssetName);
-	UNREALED_API bool ApplyCompressionAlgorithm(TArray<UAnimSequence*>& AnimSequencePtrs, UAnimCompress* Algorithm);
+
+	/** Applies the animation compression codecs to the sequence list with optional override settings */
+	UNREALED_API bool ApplyCompressionAlgorithm(TArray<UAnimSequence*>& AnimSequencePtrs, UAnimBoneCompressionSettings* OverrideSettings);
 
 	// template version of simple creating animation asset
 	template< class T >
diff --git a/Engine/Source/Editor/UnrealEd/Public/Dialogs/Dialogs.h b/Engine/Source/Editor/UnrealEd/Public/Dialogs/Dialogs.h
--- a/Engine/Source/Editor/UnrealEd/Public/Dialogs/Dialogs.h
+++ b/Engine/Source/Editor/UnrealEd/Public/Dialogs/Dialogs.h
@@ -199,7 +199,7 @@ public:
 		MyWindow = InWindow;
 	}
 
-	UNREALED_API static void OpenDialog(const FText& InDialogTitle, const TSharedRef< SWidget >& DisplayContent, const FArguments& InArgs = FArguments());
+	UNREALED_API static void OpenDialog(const FText& InDialogTitle, const TSharedRef< SWidget >& DisplayContent, const FArguments& InArgs = FArguments(), bool bAsModalDialog = false);
 
 private:
 	FReply OnOK_Clicked(void);
diff --git a/Engine/Source/Editor/UnrealEd/Public/EditorModeManager.h b/Engine/Source/Editor/UnrealEd/Public/EditorModeManager.h
--- a/Engine/Source/Editor/UnrealEd/Public/EditorModeManager.h
+++ b/Engine/Source/Editor/UnrealEd/Public/EditorModeManager.h
@@ -456,6 +456,13 @@ PRAGMA_ENABLE_DEPRECATION_WARNINGS
 	/**	Broadcasts the EditorModeIDChanged event */
 	void BroadcastEditorModeIDChanged(const FEditorModeID& ModeID, bool IsEnteringMode) { EditorModeIDChangedEvent.Broadcast(ModeID, IsEnteringMode); }
 
+	/** delegate type for triggering when coordinate system changed */
+	DECLARE_EVENT_OneParam(FEditorModeTools, FCoordSystemChangedEvent, ECoordSystem);
+	FCoordSystemChangedEvent& OnCoordSystemChanged() { return CoordSystemChangedEvent; }
+
+	/**	Broadcasts the CoordSystemChangedEvent event */
+	void BroadcastCoordSystemChanged(ECoordSystem InCoordSystem) { CoordSystemChangedEvent.Broadcast(InCoordSystem); }
+	
 	/**
 	 * Returns the current CoordSystem
 	 * 
@@ -629,6 +636,9 @@ private:
 	/** Multicast delegate that is broadcast when a widget mode is changed */
 	FWidgetModeChangedEvent WidgetModeChangedEvent;
 
+	/** Multicast delegate that is broadcast when the coordinate system is changed */
+	FCoordSystemChangedEvent CoordSystemChangedEvent;
+
 	/** The dock tab for any modes that generate a toolbar */
 	TWeakPtr<SDockTab> ModeToolbarTab;
 
diff --git a/Engine/Source/Editor/UnrealEd/Public/Factories.h b/Engine/Source/Editor/UnrealEd/Public/Factories.h
--- a/Engine/Source/Editor/UnrealEd/Public/Factories.h
+++ b/Engine/Source/Editor/UnrealEd/Public/Factories.h
@@ -56,6 +56,9 @@ protected:
 	/** Post handling of constructed objects by the factory */
 	virtual void PostProcessConstructedObjects() {};
 
+	/** Provide the opportunity to change the name during import */
+	virtual void UpdateObjectName(UClass* ObjectClass, FName& InOutObjName) {};
+
 	/** Util to ensure that InName is a valid name for a new object within InParent. Will rename any existing object within InParent if it is called InName. */
 	static void ClearObjectNameUsage(UObject* InParent, FName InName);
 
diff --git a/Engine/Source/Editor/UnrealEd/Public/FbxImporter.h b/Engine/Source/Editor/UnrealEd/Public/FbxImporter.h
--- a/Engine/Source/Editor/UnrealEd/Public/FbxImporter.h
+++ b/Engine/Source/Editor/UnrealEd/Public/FbxImporter.h
@@ -464,7 +464,7 @@ struct FbxSceneInfo
 /**
 * FBX basic data conversion class.
 */
-class FFbxDataConverter
+class UNREALED_API FFbxDataConverter
 {
 public:
 	static void SetJointPostConversionMatrix(FbxAMatrix ConversionMatrix) { JointPostConversionMatrix = ConversionMatrix; }
diff --git a/Engine/Source/Editor/UnrealEd/Public/PackageHelperFunctions.h b/Engine/Source/Editor/UnrealEd/Public/PackageHelperFunctions.h
--- a/Engine/Source/Editor/UnrealEd/Public/PackageHelperFunctions.h
+++ b/Engine/Source/Editor/UnrealEd/Public/PackageHelperFunctions.h
@@ -18,7 +18,7 @@
 
 class Error;
 
-DECLARE_LOG_CATEGORY_EXTERN(LogPackageHelperFunctions, Log, All);
+UNREALED_API DECLARE_LOG_CATEGORY_EXTERN(LogPackageHelperFunctions, Log, All);
 
 /**
  * Flags which modify the way that NormalizePackageNames works.
diff --git a/Engine/Source/Editor/UnrealEd/Public/ThumbnailHelpers.h b/Engine/Source/Editor/UnrealEd/Public/ThumbnailHelpers.h
--- a/Engine/Source/Editor/UnrealEd/Public/ThumbnailHelpers.h
+++ b/Engine/Source/Editor/UnrealEd/Public/ThumbnailHelpers.h
@@ -98,6 +98,9 @@ public:
 	/** Sets the skeletal mesh to use in the next GetView() */
 	void SetSkeletalMesh(class USkeletalMesh* InSkeletalMesh);
 
+	/** Returns the preview actor within the scene */
+	class ASkeletalMeshActor* GetPreviewActor() { return PreviewActor; }
+
 protected:
 	// FThumbnailPreviewScene implementation
 	virtual void GetViewMatrixParameters(const float InFOVDegrees, FVector& OutOrigin, float& OutOrbitPitch, float& OutOrbitYaw, float& OutOrbitZoom) const override;
diff --git a/Engine/Source/Programs/UnrealHeaderTool/Private/HeaderParser.cpp b/Engine/Source/Programs/UnrealHeaderTool/Private/HeaderParser.cpp
--- a/Engine/Source/Programs/UnrealHeaderTool/Private/HeaderParser.cpp
+++ b/Engine/Source/Programs/UnrealHeaderTool/Private/HeaderParser.cpp
@@ -6604,6 +6604,11 @@ void FHeaderParser::ParseRigVMMethodParameters(UStruct* Struct)
 		StructRigVMInfo->Members.Add(MoveTemp(Parameter));
 	}
 
+	if (StructRigVMInfo->Members.Num() == 0)
+	{
+		UE_LOG_ERROR_UHT(TEXT("RigVM Struct '%s' - has zero members - invalid RIGVM_METHOD."), *Struct->GetName());
+	}
+
 	if (StructRigVMInfo->Members.Num() > 64)
 	{
 		UE_LOG_ERROR_UHT(TEXT("RigVM Struct '%s' - has %d members (64 is the limit)."), *Struct->GetName(), StructRigVMInfo->Members.Num());
diff --git a/Engine/Source/Runtime/AnimGraphRuntime/AnimGraphRuntime.Build.cs b/Engine/Source/Runtime/AnimGraphRuntime/AnimGraphRuntime.Build.cs
--- a/Engine/Source/Runtime/AnimGraphRuntime/AnimGraphRuntime.Build.cs
+++ b/Engine/Source/Runtime/AnimGraphRuntime/AnimGraphRuntime.Build.cs
@@ -14,7 +14,7 @@ public class AnimGraphRuntime : ModuleRules
 				"Core", 
 				"CoreUObject", 
 				"Engine",
-				"AnimationCore",
+                "AnimationCore",
 			}
 		);
 
@@ -24,17 +24,22 @@ public class AnimGraphRuntime : ModuleRules
 			}
 		);
 
-		SetupModulePhysicsSupport(Target);
+        SetupModulePhysicsSupport(Target);
 
-		if (Target.bCompileChaos || Target.bUseChaos)
-		{
-			PublicDependencyModuleNames.AddRange(
-				new string[] {
+		// External users of this library do not need to know about Eigen.
+        AddEngineThirdPartyPrivateStaticDependencies(Target,
+                "Eigen"
+                );
+
+        if (Target.bCompileChaos || Target.bUseChaos)
+        {
+            PublicDependencyModuleNames.AddRange(
+                new string[] {
 					"ChaosSolvers",
 					"GeometryCollectionEngine",
 					"GeometryCollectionSimulationCore"
-				}
-			);
-		}
-	}
+                }
+            );
+        }
+    }
 }
diff --git a/Engine/Source/Runtime/AnimGraphRuntime/Private/AnimNodes/AnimNode_RandomPlayer.cpp b/Engine/Source/Runtime/AnimGraphRuntime/Private/AnimNodes/AnimNode_RandomPlayer.cpp
--- a/Engine/Source/Runtime/AnimGraphRuntime/Private/AnimNodes/AnimNode_RandomPlayer.cpp
+++ b/Engine/Source/Runtime/AnimGraphRuntime/Private/AnimNodes/AnimNode_RandomPlayer.cpp
@@ -1,18 +1,16 @@
 // Copyright Epic Games, Inc. All Rights Reserved.
 
 #include "AnimNodes/AnimNode_RandomPlayer.h"
-#include "AnimationRuntime.h"
+
+#include "Algo/BinarySearch.h"
 #include "Animation/AnimInstanceProxy.h"
 #include "Animation/AnimSequence.h"
 #include "Animation/AnimTrace.h"
 
 FAnimNode_RandomPlayer::FAnimNode_RandomPlayer()
-: CurrentEntry(INDEX_NONE)
-, NextEntry(INDEX_NONE)
-, CurrentDataIndex(0)
-, bShuffleMode(false)
+    : CurrentPlayDataIndex(0)
+    , bShuffleMode(false)
 {
-
 }
 
 void FAnimNode_RandomPlayer::Initialize_AnyThread(const FAnimationInitializeContext& Context)
@@ -21,59 +19,95 @@ void FAnimNode_RandomPlayer::Initialize_AnyThread(const FAnimationInitializeCont
 	FAnimNode_Base::Initialize_AnyThread(Context);
 	GetEvaluateGraphExposedInputs().Execute(Context);
 
-	const int32 NumEntries = Entries.Num();
+	// Create a sanitized list of valid entries and only operate on those from here on in.
+	ValidEntries.Empty(Entries.Num());
+	for (int32 EntryIndex = 0; EntryIndex < Entries.Num(); EntryIndex++)
+	{
+		FRandomPlayerSequenceEntry* Entry = &Entries[EntryIndex];
 
-	if(NumEntries == 0)
+		if (Entry->Sequence == nullptr)
+		{
+			continue;
+		}
+
+		// If the likelihood of this entry playing is nil, then skip it as well.
+		if (!bShuffleMode && Entry->ChanceToPlay <= SMALL_NUMBER)
+		{
+			continue;
+		}
+
+		ValidEntries.Push(Entry);
+	}
+
+	const int32 NumValidEntries = ValidEntries.Num();
+
+	if (NumValidEntries == 0)
 	{
 		// early out here, no need to do anything at all if we're not playing anything
 		return;
 	}
 
 	NormalizedPlayChances.Empty(NormalizedPlayChances.Num());
-	NormalizedPlayChances.AddUninitialized(NumEntries);
+	NormalizedPlayChances.AddUninitialized(NumValidEntries + 1);
 
-	// Initialize normalized play chance for each entry and validate entry data
+	// Sanitize the data and sum up the range of the random chances so that
+	// we can normalize the individual chances below.
 	float SumChances = 0.0f;
-	for(FRandomPlayerSequenceEntry& Entry : Entries)
+	for (FRandomPlayerSequenceEntry* Entry : ValidEntries)
 	{
-		SumChances += Entry.ChanceToPlay;
+		SumChances += Entry->ChanceToPlay;
 
-		if(Entry.MaxLoopCount < Entry.MinLoopCount)
+		if (Entry->MaxLoopCount < Entry->MinLoopCount)
 		{
-			Swap(Entry.MaxLoopCount, Entry.MinLoopCount);
+			Swap(Entry->MaxLoopCount, Entry->MinLoopCount);
 		}
 
-		if(Entry.MaxPlayRate < Entry.MinPlayRate)
+		if (Entry->MaxPlayRate < Entry->MinPlayRate)
 		{
-			Swap(Entry.MaxLoopCount, Entry.MinLoopCount);
+			Swap(Entry->MaxLoopCount, Entry->MinLoopCount);
 		}
+
+		Entry->BlendIn.Reset();
 	}
 
-	for(int32 Idx = 0 ; Idx < NumEntries ; ++Idx)
+	if (bShuffleMode)
 	{
-		NormalizedPlayChances[Idx] = Entries[Idx].ChanceToPlay / SumChances;
+		// Seed the shuffle list, ignoring all last entry checks, since we're doing the
+		// initial build and don't care about the non-repeatability property (yet).
+		BuildShuffleList(INDEX_NONE);
+	}
+	else
+	{
+		// Ensure that our chance sum is non-"zero" and non-negative.
+		check(SumChances > SMALL_NUMBER);
+
+		// Construct a cumulative distribution function so that we can look up the
+		// index of the sequence using binary search on the [0-1) random number.
+		float CurrentChance = 0.0f;
+		for (int32 Idx = 0; Idx < NumValidEntries; ++Idx)
+		{
+			CurrentChance += ValidEntries[Idx]->ChanceToPlay / SumChances;
+			NormalizedPlayChances[Idx] = CurrentChance;
+		}
+		NormalizedPlayChances[NumValidEntries] = 1.0f;
 	}
 
 	// Initialize random stream and pick first entry
 	RandomStream.Initialize(FPlatformTime::Cycles());
-	
-	CurrentEntry = GetNextEntryIndex();
-	NextEntry = GetNextEntryIndex();
 
 	PlayData.Empty(2);
 	PlayData.AddDefaulted(2);
 
-	FRandomAnimPlayData& CurrentData = PlayData[GetDataIndex(ERandomDataIndexType::Current)];
-	FRandomAnimPlayData& NextData = PlayData[GetDataIndex(ERandomDataIndexType::Next)];
+	int32 CurrentEntry = GetNextValidEntryIndex();
+	int32 NextEntry = GetNextValidEntryIndex();
 
-	// Init play data
-	CurrentData.BlendWeight = 1.0f;
-	CurrentData.PlayRate = RandomStream.FRandRange(Entries[CurrentEntry].MinPlayRate, Entries[CurrentEntry].MaxPlayRate);
-	CurrentData.RemainingLoops = FMath::Clamp(RandomStream.RandRange(Entries[CurrentEntry].MinLoopCount, Entries[CurrentEntry].MaxLoopCount), 0, MAX_int32);
+	// Initialize the animation data for the first and the next sequence so that we can properly
+	// blend between them.
+	FRandomAnimPlayData& CurrentData = GetPlayData(ERandomDataIndexType::Current);
+	InitPlayData(CurrentData, CurrentEntry, 1.0f);
 
-	NextData.BlendWeight = 0.0f;
-	NextData.PlayRate = RandomStream.FRandRange(Entries[NextEntry].MinPlayRate, Entries[NextEntry].MaxPlayRate);
-	NextData.RemainingLoops = FMath::Clamp(RandomStream.RandRange(Entries[NextEntry].MinLoopCount, Entries[NextEntry].MaxLoopCount), 0, MAX_int32);
+	FRandomAnimPlayData& NextData = GetPlayData(ERandomDataIndexType::Next);
+	InitPlayData(NextData, NextEntry, 0.0f);
 }
 
 void FAnimNode_RandomPlayer::Update_AnyThread(const FAnimationUpdateContext& Context)
@@ -81,160 +115,188 @@ void FAnimNode_RandomPlayer::Update_AnyThread(const FAnimationUpdateContext& Con
 	DECLARE_SCOPE_HIERARCHICAL_COUNTER_ANIMNODE(Update_AnyThread)
 	GetEvaluateGraphExposedInputs().Execute(Context);
 
-	if(Entries.Num() == 0)
+	if (ValidEntries.Num() == 0)
 	{
 		// We don't have any entries, play data will be invalid - early out
 		return;
 	}
 
-	FRandomAnimPlayData* CurrentData = &PlayData[GetDataIndex(ERandomDataIndexType::Current)];
-	FRandomAnimPlayData* NextData = &PlayData[GetDataIndex(ERandomDataIndexType::Next)];
+	FRandomAnimPlayData* CurrentData = &GetPlayData(ERandomDataIndexType::Current);
+	FRandomAnimPlayData* NextData = &GetPlayData(ERandomDataIndexType::Next);
+
+	const UAnimSequence* CurrentSequence = CurrentData->Entry->Sequence;
 
-	if(UAnimSequence* CurrentSequence = Entries[CurrentEntry].Sequence)
+	// If we looped around, adjust the previous play time to always be before the current playtime,
+	// since we can assume modulo. This makes the crossing check for the start time a lot simpler.
+	float AdjustedPreviousPlayTime = CurrentData->PreviousPlayTime;
+	if (CurrentData->CurrentPlayTime < AdjustedPreviousPlayTime)
 	{
-		float TimeRemaining = CurrentSequence->SequenceLength - CurrentData->InternalTimeAccumulator;
+		AdjustedPreviousPlayTime -= CurrentSequence->SequenceLength;
+	}
 
-		if(CurrentData->InternalTimeAccumulator < CurrentData->PreviousTimeAccumulator)
+	// Did we cross the play start time? Decrement the loop counter. Once we're on the last loop, we can
+	// start blending into the next animation.
+	if (AdjustedPreviousPlayTime < CurrentData->PlayStartTime && CurrentData->PlayStartTime <= CurrentData->CurrentPlayTime)
+	{
+		// We've looped, update remaining
+		--CurrentData->RemainingLoops;
+	}
+
+	bool bAdvanceToNextEntry = false;
+
+	if (CurrentData->RemainingLoops <= 0)
+	{
+		const bool bNextAnimIsDifferent = CurrentData->Entry != NextData->Entry;
+
+		// If we're in the blend window start blending, but only if we're moving to a new animation,
+		// otherwise just keep looping.
+		FRandomPlayerSequenceEntry& NextSequenceEntry = *NextData->Entry;
+
+		// If the next animation is different, then smoothly blend between them. Otherwise
+		// we do a hard transition to the same play point. The next animation might play at
+		// a different rate, so we have to switch.
+		if (bNextAnimIsDifferent)
 		{
-			// We've looped, update remaining
-			--CurrentData->RemainingLoops;
+			bool bDoBlending = false;
 
-			if(CurrentData->RemainingLoops < 0)
+			// Are we already blending? Continue to do so.
+			if (FAnimationRuntime::HasWeight(NextSequenceEntry.BlendIn.GetAlpha()))
+			{
+				bDoBlending = true;
+			}
+			else
 			{
-				// If we're switching to the same anim
-				if(CurrentEntry == NextEntry)
+				// Check to see if we need to start the blending process.
+				float AmountPlayedSoFar = CurrentData->CurrentPlayTime - CurrentData->PlayStartTime;
+				if (AmountPlayedSoFar < 0.0f)
 				{
-					// Need to switch to the next anim, but first put our accumulator in the next data we're about to switch
-					// to so we don't see a pop
-					NextData->InternalTimeAccumulator = CurrentData->InternalTimeAccumulator;
+					AmountPlayedSoFar += CurrentSequence->SequenceLength;
 				}
 
-				SwitchNextToCurrent();
+				float TimeRemaining = CurrentSequence->SequenceLength - AmountPlayedSoFar;
 
-				// Re-get data as we've switched over
-				CurrentData = &PlayData[GetDataIndex(ERandomDataIndexType::Current)];
-				NextData = &PlayData[GetDataIndex(ERandomDataIndexType::Next)];
+				if (TimeRemaining <= NextSequenceEntry.BlendIn.GetBlendTime())
+				{
+					bDoBlending = true;
+				}
 			}
-		}
-
-		// Cache time to detect loops
-		CurrentData->PreviousTimeAccumulator = CurrentData->InternalTimeAccumulator;
-		NextData->PreviousTimeAccumulator = NextData->InternalTimeAccumulator;
-
-		// If we're in the blend window start blending, but only if we're moving to a new animation,
-		// otherwise just keep looping.
-		const bool bInCrossfadeTime = TimeRemaining <= Entries[NextEntry].BlendIn.GetBlendTime();
-		const bool bNextAnimIsDifferent = NextEntry != CurrentEntry;
-		const bool bNeedMoreLoops = CurrentData->RemainingLoops > 0;
 
-		if(bInCrossfadeTime && !bNeedMoreLoops)
-		{
-			if(bNextAnimIsDifferent)
+			if (bDoBlending)
 			{
 				// Blending to next
-				Entries[NextEntry].BlendIn.Update(Context.GetDeltaTime());
+				NextSequenceEntry.BlendIn.Update(Context.GetDeltaTime());
 
-				float BlendedAlpha = Entries[NextEntry].BlendIn.GetBlendedValue();
-
-				if(BlendedAlpha < 1.0f)
+				if (NextSequenceEntry.BlendIn.IsComplete())
+				{
+					// Set the play start time to be the current play time so that loop counts are properly
+					// maintained.
+					NextData->PlayStartTime = NextData->CurrentPlayTime;
+					bAdvanceToNextEntry = true;
+				}
+				else
 				{
-					NextData->BlendWeight = BlendedAlpha;
-					CurrentData->BlendWeight = 1.0f - BlendedAlpha;
+					float BlendedAlpha = NextSequenceEntry.BlendIn.GetBlendedValue();
+
+					if (BlendedAlpha < 1.0f)
+					{
+						NextData->BlendWeight = BlendedAlpha;
+						CurrentData->BlendWeight = 1.0f - BlendedAlpha;
+					}
 				}
 			}
 		}
-
-		// If we were blending but now we're done, switch play data
-		if(Entries[NextEntry].BlendIn.IsComplete())
+		else if (!bNextAnimIsDifferent && CurrentData->RemainingLoops < 0)
 		{
-			SwitchNextToCurrent();
+			NextData->CurrentPlayTime = CurrentData->CurrentPlayTime;
 
-			// Re-get data as we've switched over
-			CurrentData = &PlayData[GetDataIndex(ERandomDataIndexType::Current)];
-			NextData = &PlayData[GetDataIndex(ERandomDataIndexType::Next)];
+			// Set the play start time to be the current play time so that loop counts are properly
+			// maintained.
+			NextData->PlayStartTime = NextData->CurrentPlayTime;
+			bAdvanceToNextEntry = true;
 		}
+	}
 
-		if(FAnimInstanceProxy* AnimProxy = Context.AnimInstanceProxy)
-		{
-			FAnimGroupInstance* SyncGroup;
-			FAnimTickRecord& TickRecord = AnimProxy->CreateUninitializedTickRecord(INDEX_NONE, SyncGroup);
-			AnimProxy->MakeSequenceTickRecord(TickRecord, Entries[CurrentEntry].Sequence, true, CurrentData->PlayRate, CurrentData->BlendWeight, CurrentData->InternalTimeAccumulator, CurrentData->MarkerTickRecord);
+	// Cache time to detect loops
+	CurrentData->PreviousPlayTime = CurrentData->CurrentPlayTime;
+	NextData->PreviousPlayTime = NextData->CurrentPlayTime;
 
-			TRACE_ANIM_TICK_RECORD(Context, TickRecord);
+	if (bAdvanceToNextEntry)
+	{
+		AdvanceToNextSequence();
 
-			if(NextData->BlendWeight > 0.0f)
-			{
-				FAnimTickRecord& NextTickRecord = AnimProxy->CreateUninitializedTickRecord(INDEX_NONE, SyncGroup);
-				AnimProxy->MakeSequenceTickRecord(NextTickRecord, Entries[NextEntry].Sequence, true, NextData->PlayRate, NextData->BlendWeight, NextData->InternalTimeAccumulator, NextData->MarkerTickRecord);
+		// Re-get data as we've switched over
+		CurrentData = &GetPlayData(ERandomDataIndexType::Current);
+		NextData = &GetPlayData(ERandomDataIndexType::Next);
+	}
 
-				TRACE_ANIM_TICK_RECORD(Context, NextTickRecord);
-			}
-		}
+	FAnimInstanceProxy* AnimProxy = Context.AnimInstanceProxy;
+	FAnimGroupInstance* SyncGroup;
+	FAnimTickRecord& TickRecord = AnimProxy->CreateUninitializedTickRecord(INDEX_NONE, SyncGroup);
+	AnimProxy->MakeSequenceTickRecord(
+	    TickRecord, CurrentData->Entry->Sequence, true, CurrentData->PlayRate,
+	    CurrentData->BlendWeight, CurrentData->CurrentPlayTime, CurrentData->MarkerTickRecord);
+	
+	TRACE_ANIM_TICK_RECORD(Context, TickRecord);
+
+	if (FAnimationRuntime::HasWeight(NextData->BlendWeight))
+	{
+		FAnimTickRecord& NextTickRecord = AnimProxy->CreateUninitializedTickRecord(INDEX_NONE, SyncGroup);
+		AnimProxy->MakeSequenceTickRecord(
+		    NextTickRecord, NextData->Entry->Sequence, true, NextData->PlayRate,
+		    NextData->BlendWeight, NextData->CurrentPlayTime, NextData->MarkerTickRecord);
+			
+		TRACE_ANIM_TICK_RECORD(Context, NextTickRecord);
 	}
 
-	TRACE_ANIM_NODE_VALUE(Context, TEXT("Current Sequence"), Entries[CurrentEntry].Sequence);
+	TRACE_ANIM_NODE_VALUE(Context, TEXT("Current Sequence"), CurrentData ? CurrentData->Entry->Sequence : nullptr);
 	TRACE_ANIM_NODE_VALUE(Context, TEXT("Current Weight"), CurrentData ? CurrentData->BlendWeight : 0.0f);
-	TRACE_ANIM_NODE_VALUE(Context, TEXT("Next Sequence"), Entries[NextEntry].Sequence);
+	TRACE_ANIM_NODE_VALUE(Context, TEXT("Next Sequence"), NextData ? NextData->Entry->Sequence : nullptr);
 	TRACE_ANIM_NODE_VALUE(Context, TEXT("Next Weight"), NextData ? NextData->BlendWeight : 0.0f);
 }
 
 void FAnimNode_RandomPlayer::Evaluate_AnyThread(FPoseContext& Output)
 {
 	DECLARE_SCOPE_HIERARCHICAL_COUNTER_ANIMNODE(Evaluate_AnyThread)
-	if(Entries.Num() > 0)
+	if (ValidEntries.Num() == 0)
 	{
-		UAnimSequence* CurrentSequence = Entries[CurrentEntry].Sequence;
+		return;
+	}
 
-		if(CurrentSequence)
-		{
-			FRandomAnimPlayData& CurrentData = PlayData[GetDataIndex(ERandomDataIndexType::Current)];
-			FRandomAnimPlayData& NextData = PlayData[GetDataIndex(ERandomDataIndexType::Next)];
+	FRandomAnimPlayData& CurrentData = GetPlayData(ERandomDataIndexType::Current);
+	FRandomAnimPlayData& NextData = GetPlayData(ERandomDataIndexType::Next);
 
-			if(CurrentData.BlendWeight != 1.0f)
-			{
-				if(FAnimInstanceProxy* AnimProxy = Output.AnimInstanceProxy)
-				{
-					// Start Blending
-					FCompactPose Poses[2];
-					FBlendedCurve Curves[2];
-					float Weights[2];
+	UAnimSequence* CurrentSequence = CurrentData.Entry->Sequence;
 
-					const FBoneContainer& RequiredBone = AnimProxy->GetRequiredBones();
-					Poses[0].SetBoneContainer(&RequiredBone);
-					Poses[1].SetBoneContainer(&RequiredBone);
+	if (!FMath::IsNearlyEqualByULP(CurrentData.BlendWeight, 1.0f))
+	{
+		FAnimInstanceProxy* AnimProxy = Output.AnimInstanceProxy;
 
-					Curves[0].InitFrom(RequiredBone);
-					Curves[1].InitFrom(RequiredBone);
+		// Start Blending
+		FCompactPose Poses[2];
+		FBlendedCurve Curves[2];
+		float Weights[2];
 
-					Weights[0] = CurrentData.BlendWeight;
-					Weights[1] = NextData.BlendWeight;
+		const FBoneContainer& RequiredBone = AnimProxy->GetRequiredBones();
+		Poses[0].SetBoneContainer(&RequiredBone);
+		Poses[1].SetBoneContainer(&RequiredBone);
 
-					UAnimSequence* NextSequence = Entries[NextEntry].Sequence;
+		Curves[0].InitFrom(RequiredBone);
+		Curves[1].InitFrom(RequiredBone);
 
-					CurrentSequence->GetAnimationPose(Poses[0], Curves[0], FAnimExtractContext(CurrentData.InternalTimeAccumulator, AnimProxy->ShouldExtractRootMotion()));
-					NextSequence->GetAnimationPose(Poses[1], Curves[1], FAnimExtractContext(NextData.InternalTimeAccumulator, AnimProxy->ShouldExtractRootMotion()));
+		Weights[0] = CurrentData.BlendWeight;
+		Weights[1] = NextData.BlendWeight;
 
-					FAnimationRuntime::BlendPosesTogether(Poses, Curves, Weights, Output.Pose, Output.Curve);
-				}
-				else
-				{
-					Output.ResetToRefPose();
-				}
-			}
-			else
-			{
-				// Single anim
-				CurrentSequence->GetAnimationPose(Output.Pose, Output.Curve, FAnimExtractContext(CurrentData.InternalTimeAccumulator, Output.AnimInstanceProxy->ShouldExtractRootMotion()));
-			}
-		}
-		else
-		{
-			Output.ResetToRefPose();
-		}
+		UAnimSequence* NextSequence = NextData.Entry->Sequence;
+
+		CurrentSequence->GetAnimationPose(Poses[0], Curves[0], FAnimExtractContext(CurrentData.CurrentPlayTime, AnimProxy->ShouldExtractRootMotion()));
+		NextSequence->GetAnimationPose(Poses[1], Curves[1], FAnimExtractContext(NextData.CurrentPlayTime, AnimProxy->ShouldExtractRootMotion()));
+
+		FAnimationRuntime::BlendPosesTogether(Poses, Curves, Weights, Output.Pose, Output.Curve);
 	}
 	else
 	{
-		Output.ResetToRefPose();
+		// Single animation, no blending needed.
+		CurrentSequence->GetAnimationPose(Output.Pose, Output.Curve, FAnimExtractContext(CurrentData.CurrentPlayTime, Output.AnimInstanceProxy->ShouldExtractRootMotion()));
 	}
 }
 
@@ -246,109 +308,106 @@ void FAnimNode_RandomPlayer::GatherDebugData(FNodeDebugData& DebugData)
 	DebugData.AddDebugItem(DebugLine, true);
 }
 
-int32 FAnimNode_RandomPlayer::GetNextEntryIndex()
+int32 FAnimNode_RandomPlayer::GetNextValidEntryIndex()
 {
-	if(Entries.Num() > 0)
+	check(ValidEntries.Num() > 0);
+
+	if (bShuffleMode)
 	{
-		if(bShuffleMode)
-		{
-			if(ShuffleList.Num() == 0)
-			{
-				// Need a new list
-				BuildShuffleList();
-			}
+		// Get the top value, don't allow realloc
+		int32 Index = ShuffleList.Pop(false);
 
-			// Get the top value, don't allow realloc
-			return ShuffleList.Pop(false);
-		}
-		else
+		// If we cleared the shuffles, rebuild for the next round, indicating
+		// the current value so that we don't pop that one off again next time.
+		if (ShuffleList.Num() == 0)
 		{
-			float RandomVal = RandomStream.GetFraction();
-			const int32 NumEntries = Entries.Num();
-
-			// Grab the entry index corresponding to the value
-			for(int32 Idx = 0 ; Idx < NumEntries ; ++Idx)
-			{
-				RandomVal -= NormalizedPlayChances[Idx];
-				if(RandomVal <= 0.0f)
-				{
-					return Idx;
-				}
-			}
+			BuildShuffleList(Index);
 		}
+
+		return Index;
 	}
+	else
+	{
+		float RandomVal = RandomStream.GetFraction();
 
-	return INDEX_NONE;
+		// Search the cumulative distribution array for the last entry that's
+		// smaller or equal to the random value. That becomes our new animation.
+		return Algo::UpperBound(NormalizedPlayChances, RandomVal);
+	}
 }
 
-int32 FAnimNode_RandomPlayer::GetDataIndex(const ERandomDataIndexType& Type)
+FRandomAnimPlayData& FAnimNode_RandomPlayer::GetPlayData(ERandomDataIndexType Type)
 {
-	if(Type == ERandomDataIndexType::Current)
+	// PlayData only holds two entries. We swap between them in AdvanceToNextSequence
+	// by setting CUrrentPlayDataIndex to either 0 or 1. Hence the modulo 2 magic below.
+	if (Type == ERandomDataIndexType::Current)
 	{
-		return CurrentDataIndex;
+		return PlayData[CurrentPlayDataIndex];
 	}
 	else
 	{
-		// Next Accumulator
-		return (CurrentDataIndex + 1) % 2;
+		return PlayData[(CurrentPlayDataIndex + 1) % 2];
 	}
 }
 
-void FAnimNode_RandomPlayer::SwitchNextToCurrent()
+void FAnimNode_RandomPlayer::InitPlayData(FRandomAnimPlayData& Data, int32 ValidEntryIndex, float BlendWeight)
 {
-	// reset alpha blend we've possibly just taken
-	Entries[NextEntry].BlendIn.Reset();
+	FRandomPlayerSequenceEntry* Entry = ValidEntries[ValidEntryIndex];
 
-	// Switch which entry to get sequences and parameters from, and pre-generate the next entry index
-	CurrentEntry = NextEntry;
-	NextEntry = GetNextEntryIndex();
+	Data.Entry = Entry;
+	Data.BlendWeight = BlendWeight;
+	Data.PlayRate = RandomStream.FRandRange(Entry->MinPlayRate, Entry->MaxPlayRate);
+	Data.RemainingLoops = FMath::Clamp(RandomStream.RandRange(Entry->MinLoopCount, Entry->MaxLoopCount), 0, MAX_int32);
 
-	// Switch play data
-	CurrentDataIndex = (CurrentDataIndex + 1) %2;
+	Data.PlayStartTime = 0.0f;
+	Data.CurrentPlayTime = 0.0f;
+	Data.PreviousPlayTime = 0.0f;
+	Data.MarkerTickRecord.Reset();
+}
+
+void FAnimNode_RandomPlayer::AdvanceToNextSequence()
+{
+	// Get the next sequence entry to use.
+	int32 NextEntry = GetNextValidEntryIndex();
+
+	// Switch play data by flipping it between 0 and 1.
+	CurrentPlayDataIndex = (CurrentPlayDataIndex + 1) % 2;
 
 	// Get our play data
-	FRandomAnimPlayData& CurrentData = PlayData[GetDataIndex(ERandomDataIndexType::Current)];
-	FRandomAnimPlayData& NextData = PlayData[GetDataIndex(ERandomDataIndexType::Next)];
+	FRandomAnimPlayData& CurrentData = GetPlayData(ERandomDataIndexType::Current);
+	FRandomAnimPlayData& NextData = GetPlayData(ERandomDataIndexType::Next);
 
-	// Reset blendweights
+	// Reset blend weights
 	CurrentData.BlendWeight = 1.0f;
-	NextData.BlendWeight = 0.0f;
+	CurrentData.Entry->BlendIn.Reset();
 
 	// Set up data for next switch
-	NextData.InternalTimeAccumulator = 0.0f;
-	NextData.PreviousTimeAccumulator = 0.0f;
-	NextData.PlayRate = RandomStream.FRandRange(Entries[NextEntry].MinPlayRate, Entries[NextEntry].MaxPlayRate);
-	NextData.RemainingLoops = FMath::Clamp(RandomStream.RandRange(Entries[NextEntry].MinLoopCount, Entries[NextEntry].MaxLoopCount), 0, MAX_int32);
-	NextData.MarkerTickRecord.Reset();
+	InitPlayData(NextData, NextEntry, 0.0f);
 }
 
-void FAnimNode_RandomPlayer::BuildShuffleList()
+void FAnimNode_RandomPlayer::BuildShuffleList(int32 LastEntry)
 {
-	ShuffleList.Reset(Entries.Num());
+	ShuffleList.Reset(ValidEntries.Num());
 
 	// Build entry index list
-	const int32 NumEntries = Entries.Num();
-	for(int32 i = 0 ; i < NumEntries ; ++i)
+	const int32 NumValidEntries = ValidEntries.Num();
+	for (int32 i = 0; i < NumValidEntries; ++i)
 	{
 		ShuffleList.Add(i);
 	}
 
 	// Shuffle the list
 	const int32 NumShuffles = ShuffleList.Num() - 1;
-	for(int32 i = 0 ; i < NumShuffles ; ++i)
+	for (int32 i = 0; i < NumShuffles; ++i)
 	{
 		int32 SwapIdx = RandomStream.RandRange(i, NumShuffles);
 		ShuffleList.Swap(i, SwapIdx);
 	}
 
-	if(ShuffleList.Num() > 1)
+	// Make sure we don't play the same thing twice in a row
+	if (ShuffleList.Num() > 1 && ShuffleList.Last() == LastEntry)
 	{
-		// Make sure we don't play the same thing twice, one at the end and one at the beginning of
-		// the list
-		if(ShuffleList.Last() == CurrentEntry)
-		{
-			// Swap first and last
-			ShuffleList.Swap(0, ShuffleList.Num() - 1);
-		}
+		// Swap the last with a random entry.
+		ShuffleList.Swap(RandomStream.RandRange(0, ShuffleList.Num() - 2), ShuffleList.Num() - 1);
 	}
 }
diff --git a/Engine/Source/Runtime/AnimGraphRuntime/Private/RBF/RBFSolver.cpp b/Engine/Source/Runtime/AnimGraphRuntime/Private/RBF/RBFSolver.cpp
--- a/Engine/Source/Runtime/AnimGraphRuntime/Private/RBF/RBFSolver.cpp
+++ b/Engine/Source/Runtime/AnimGraphRuntime/Private/RBF/RBFSolver.cpp
@@ -1,26 +1,40 @@
 // Copyright Epic Games, Inc. All Rights Reserved.
 
 #include "RBF/RBFSolver.h"
+
+#include "RBF/RBFInterpolator.h"
+
 #include "EngineLogs.h"
 
-FQuat FRBFEntry::AsQuat(int32 Index) const
+#include "Containers/Set.h"
+
+FRotator FRBFEntry::AsRotator(int32 Index) const
 {
-	FQuat Result = FQuat::Identity;
+	FRotator Result = FRotator::ZeroRotator;
 
 	const int32 BaseIndex = Index * 3;
 
 	if (Values.Num() >= BaseIndex + 3)
 	{
-		FRotator Rot;
-		Rot.Roll	= Values[BaseIndex + 0];
-		Rot.Pitch	= Values[BaseIndex + 1];
-		Rot.Yaw		= Values[BaseIndex + 2];
-		Result = FQuat(Rot);
+		Result.Roll = Values[BaseIndex + 0];
+		Result.Pitch = Values[BaseIndex + 1];
+		Result.Yaw = Values[BaseIndex + 2];
 	}
-
 	return Result;
 }
 
+FQuat FRBFEntry::AsQuat(int32 Index) const
+{
+	return AsRotator(Index).Quaternion();
+}
+
+FVector FRBFEntry::AsVector(int32 Index) const
+{
+	return AsRotator(Index).Vector();
+}
+
+
+
 void FRBFEntry::AddFromRotator(const FRotator& InRot)
 {
 	const int32 BaseIndex = Values.AddUninitialized(3);
@@ -43,6 +57,7 @@ void FRBFEntry::AddFromVector(const FVector& InVector)
 
 FRBFParams::FRBFParams()
 	: TargetDimensions(3)
+	, SolverType(ERBFSolverType::Additive)
 	, Radius(1.f)
 	, Function(ERBFFunctionType::Gaussian)
 	, DistanceMethod(ERBFDistanceMethod::Euclidean)
@@ -70,79 +85,187 @@ FVector FRBFParams::GetTwistAxisVector() const
 	}
 }
 
+
 //////////////////////////////////////////////////////////////////////////
 
-float FRBFSolver::FindDistanceBetweenEntries(const FRBFEntry& A, const FRBFEntry& B, const FRBFParams& Params, ERBFDistanceMethod OverrideMethod)
+/* Returns the distance between entries, using different metrics, in radians. */
+static float GetDistanceBetweenEntries(
+	const FRBFEntry& A,
+	const FRBFEntry& B,
+	ERBFDistanceMethod DistanceMetric,
+	const FVector &TwistAxis
+)
 {
 	check(A.GetDimensions() == B.GetDimensions());
 
-	ERBFDistanceMethod DistanceMethod = OverrideMethod == ERBFDistanceMethod::DefaultMethod ? Params.DistanceMethod : OverrideMethod;
+	const int32 NumRots = A.GetDimensions() / 3;
+	float TotalDistance = 0.0f;
 
-	// Simple n-dimensional distance
-	if (DistanceMethod == ERBFDistanceMethod::Euclidean)
+	for (int32 i = 0; i < NumRots; i++)
 	{
-		float DistSqr = 0.f;
-
-		for (int32 i = 0; i < A.Values.Num(); i++)
+		float Distance = 0.0f;
+		switch (DistanceMetric)
 		{
-			DistSqr += FMath::Square(A.Values[i] - B.Values[i]);
+		case ERBFDistanceMethod::Euclidean:
+			Distance = RBFDistanceMetric::Euclidean(A.AsRotator(i), B.AsRotator(i));
+			break;
+
+		case ERBFDistanceMethod::Quaternion:
+			Distance = RBFDistanceMetric::ArcLength(A.AsQuat(i), B.AsQuat(i));
+			break;
+
+		case ERBFDistanceMethod::SwingAngle:
+		case ERBFDistanceMethod::DefaultMethod:
+			Distance += RBFDistanceMetric::SwingAngle(A.AsQuat(i), B.AsQuat(i), TwistAxis);
+			break;
+
+		case ERBFDistanceMethod::TwistAngle:
+			Distance += RBFDistanceMetric::TwistAngle(A.AsQuat(i), B.AsQuat(i), TwistAxis);
+			break;
 		}
 
-		return FMath::Sqrt(DistSqr);
+		TotalDistance += FMath::Square(Distance);
 	}
-	// Treat values as sequence of eulers - find quat distance between each pair, then sqrt-sum-of-squares of those
-	else if (DistanceMethod == ERBFDistanceMethod::Quaternion)
-	{
-		float DistSqr = 0.f;
 
-		const int32 NumRots = A.GetDimensions() / 3;
-		for (int32 RotIdx = 0; RotIdx < NumRots; RotIdx++)
-		{
-			float RadDist = A.AsQuat(RotIdx).AngularDistance(B.AsQuat(RotIdx));
-			DistSqr += FMath::Square(FMath::RadiansToDegrees(RadDist));
-		}
+	return FMath::Sqrt(TotalDistance);
+}
 
-		return FMath::Sqrt(DistSqr);
-	}
-	// Treat values as sequence of eulers - find 'swing' distance between each pair, then sqrt-sum-of-squares of those
-	else if(DistanceMethod == ERBFDistanceMethod::SwingAngle || DistanceMethod == ERBFDistanceMethod::DefaultMethod)
-	{
-		float DistSqr = 0.f;
 
-		const int32 NumRots = A.GetDimensions() / 3;
-		for (int32 RotIdx = 0; RotIdx < NumRots; RotIdx++)
+float FRBFSolver::FindDistanceBetweenEntries(const FRBFEntry& A, const FRBFEntry& B, const FRBFParams& Params, ERBFDistanceMethod OverrideMethod)
+{
+	ERBFDistanceMethod DistanceMethod = OverrideMethod == ERBFDistanceMethod::DefaultMethod ? Params.DistanceMethod : OverrideMethod;
+
+	float Distance = GetDistanceBetweenEntries(A, B, DistanceMethod, Params.GetTwistAxisVector());
+	return FMath::RadiansToDegrees(Distance);
+}
+
+
+// Sigma controls the falloff width. The larger the value the narrower the falloff
+static float GetWeightedValue(
+	float Value, 
+	float Sigma, 
+	ERBFFunctionType FalloffFunctionType,
+	bool bBackCompFix = false
+)
+{
+	if (ensure(Value >= 0.0f))
+	{
+		switch (FalloffFunctionType)
 		{
-			FVector TwistVector = Params.GetTwistAxisVector();
-			FVector VecA = A.AsQuat(RotIdx).RotateVector(TwistVector);
-			FVector VecB = B.AsQuat(RotIdx).RotateVector(TwistVector);
+		case ERBFFunctionType::Linear:
+		case ERBFFunctionType::DefaultFunction:
+		default:
+			return RBFKernel::Linear(Value, Sigma);
 
-			const float Dot = FVector::DotProduct(VecA, VecB);
-			const float RadDist = FMath::Acos(Dot);
-			DistSqr += FMath::Square(FMath::RadiansToDegrees(RadDist));
-		}
+		case ERBFFunctionType::Gaussian:
+			if (bBackCompFix)
+			{
+				// This is how the old code formulated it. It has a much wider falloff than the
+				// one below it.
+				return FMath::Exp(-Value * Value);
+			}
+			else
+			{
+				return RBFKernel::Gaussian(Value, Sigma);
+			}
+
+		case ERBFFunctionType::Exponential:
+			if (bBackCompFix)
+			{
+				// This is how the old code formulated it. It has a much wider falloff than the
+				// one below it.
+				return 1.f / FMath::Exp(Value);
+			}
+			else
+			{
+				return RBFKernel::Exponential(Value, Sigma);
+			}
+
+		case ERBFFunctionType::Cubic:
+			return RBFKernel::Cubic(Value, Sigma);
 
-		return FMath::Sqrt(DistSqr);
+		case ERBFFunctionType::Quintic:
+			return RBFKernel::Quintic(Value, Sigma);
+		}
 	}
 	else
 	{
-		ensureMsgf(false, TEXT("Unknown ERBFDistanceMethod"));
-		return 0.f;
+		return 0.0f;
 	}
 }
 
-void FRBFSolver::Solve(const FRBFParams& Params, const TArray<FRBFTarget>& Targets, const FRBFEntry& Input, TArray<FRBFOutputWeight>& OutputWeights)
+static auto InterpolativeWeightFunction(
+	const FRBFParams& Params
+	)
 {
-	if (Params.TargetDimensions != Input.GetDimensions())
+	// This is fairly arbitrary, but is done to maintain a close relationship with how the radius 
+	// controls the falloff distance in the additive solver.
+	float Sigma = FMath::DegreesToRadians(Params.Radius);
+	FVector TwistAxis = Params.GetTwistAxisVector();
+
+	return [Sigma, TwistAxis, &Params](const FRBFEntry& A, const FRBFEntry& B) {
+		float Distance = GetDistanceBetweenEntries(A, B, Params.DistanceMethod, TwistAxis);
+		return GetWeightedValue(Distance, Sigma, Params.Function);
+	};
+}
+
+static bool ValidateInterpolative(
+	const FRBFParams& Params,
+	const TArray<FRBFTarget>& Targets,
+	TArray<int>& InvalidTargets
+)
+{
+	TArray<FRBFEntry> EntryTargets;
+	for (const auto& T : Targets)
+		EntryTargets.Add(T);
+
+	TArray<TTuple<int, int>> InvalidPairs;
+	TSet<int> InvalidTargetSet;
+
+	InvalidTargets.Empty();
+	if (TRBFInterpolator<FRBFEntry>::GetIdenticalNodePairs(EntryTargets, InterpolativeWeightFunction(Params), InvalidPairs))
 	{
-		UE_LOG(LogAnimation, Warning, TEXT("Input dimensionality is %d, expected %d"), Input.GetDimensions(), Params.TargetDimensions);
-		return;
+		// We mark the second of the pair to be invalid. Given how GetInvalidNodePairs iterates over all possible pairs,
+		// this should guarantee to catch them all.
+		for (const auto& IP : InvalidPairs)
+			InvalidTargetSet.Add(IP.Get<1>());
+
+		for (const auto& IT : InvalidTargetSet)
+			InvalidTargets.Add(IT);
+
+		// Return things in a nice sorted order, rather than TSet's hash order.
+		InvalidTargets.Sort();
 	}
 
-	TArray<float> AllWeights;
-	AllWeights.AddZeroed(Targets.Num());
+	return InvalidTargets.Num() == 0;
+}
+
+bool FRBFSolver::ValidateTargets(
+	const FRBFParams& Params,
+	const TArray<FRBFTarget>& Targets,
+	TArray<int>& InvalidTargets
+)
+{
+	switch (Params.SolverType)
+	{
+	case ERBFSolverType::Additive:
+	default:
+		// The additive solver does not care
+		return true;
+
+	case ERBFSolverType::Interpolative:
+		return ValidateInterpolative(Params, Targets, InvalidTargets);
+	}
+}
 
-	float TotalWeight = 0.f; // Keep track of total weight generated, to renormalize at the end
 
+static void SolveAdditive(
+	const FRBFParams& Params,
+	const TArray<FRBFTarget>& Targets,
+	const FRBFEntry& Input,
+	TArray<float>& AllWeights
+	)
+{
 	// Iterate over each pose, adding its contribution
 	for (int32 TargetIdx = 0; TargetIdx < Targets.Num(); TargetIdx++)
 	{
@@ -150,32 +273,13 @@ void FRBFSolver::Solve(const FRBFParams& Params, const TArray<FRBFTarget>& Targe
 		ERBFFunctionType FunctionType = Target.FunctionType == ERBFFunctionType::DefaultFunction ? Params.Function : Target.FunctionType;
 
 		// Find distance
-		const float Distance = FindDistanceBetweenEntries(Target, Input, Params, Target.DistanceMethod);
-		const float Scaling = GetRadiusForTarget(Target, Params);
+		const float Distance = FRBFSolver::FindDistanceBetweenEntries(Target, Input, Params, Target.DistanceMethod);
+		const float Scaling = FRBFSolver::GetRadiusForTarget(Target, Params);
 		const float X = Distance / Scaling;
 
-		// Evaluate radial basis function to find weight
-		float Weight = 0.f;
-		if (FunctionType == ERBFFunctionType::Gaussian)
-		{
-			Weight = FMath::Exp(-(X * X));
-		}
-		else if (FunctionType == ERBFFunctionType::Exponential)
-		{
-			Weight = 1.f / FMath::Exp(X);
-		}
-		else if (FunctionType == ERBFFunctionType::Linear || FunctionType == ERBFFunctionType::DefaultFunction)
-		{
-			Weight = FMath::Max(1.f - X, 0.f);
-		}
-		else if (FunctionType == ERBFFunctionType::Cubic)
-		{
-			Weight = FMath::Max(1.f - (X * X * X), 0.f);
-		}
-		else if (FunctionType == ERBFFunctionType::Quintic)
-		{
-			Weight = FMath::Max(1.f - (X * X * X * X * X), 0.f);
-		}
+		// Evaluate radial basis function to find weight. We default to sigma = 1.0 and scale instead
+		// using the radius value. We use the old formulation for Gauss + 
+		float Weight = GetWeightedValue(X, 1.0f, FunctionType, /*BackCompFix=*/ true);
 
 		// Apply custom curve if desired
 		if (Target.bApplyCustomCurve)
@@ -185,8 +289,68 @@ void FRBFSolver::Solve(const FRBFParams& Params, const TArray<FRBFTarget>& Targe
 
 		// Add to array of all weights. Don't threshold yet, wait for normalization step.
 		AllWeights[TargetIdx] = Weight;
+	}
+}
+
+
+static void SolveInterpolative(
+	const FRBFParams& Params,
+	const TArray<FRBFTarget>& Targets,
+	const FRBFEntry& Input,
+	TArray<float> &AllWeights
+	)
+{
+	check(Input.GetDimensions() == 3);
 
-		// Add weight to total
+	TArray<FRBFEntry> EntryTargets;
+	for (const auto& T : Targets)
+		EntryTargets.Add(T);
+
+	// FIXME: We ought to be able to store the initial RBF interpolator matrix and re-use it between solves
+	// but that requires larger changes in the PoseDriver and how this code is wrapped.
+	TRBFInterpolator<FRBFEntry> Rbf(EntryTargets, InterpolativeWeightFunction(Params));
+
+	Rbf.Interpolate(AllWeights, Input);
+
+	// Scale the weight by the scale factor on the target.
+	for (int32 i = 0; i < Targets.Num(); i++)
+	{
+		AllWeights[i] *= Targets[i].ScaleFactor;
+	}
+
+}
+
+
+void FRBFSolver::Solve(
+	const FRBFParams& Params, 
+	const TArray<FRBFTarget>& Targets, 
+	const FRBFEntry& Input, 
+	TArray<FRBFOutputWeight>& OutputWeights
+	)
+{
+	if (!ensure(Params.TargetDimensions == Input.GetDimensions()))
+	{
+		return;
+	}
+
+	TArray<float> AllWeights;
+	AllWeights.AddZeroed(Targets.Num());
+
+	switch (Params.SolverType)
+	{
+	case ERBFSolverType::Additive:
+	default:
+		SolveAdditive(Params, Targets, Input, AllWeights);
+		break;
+
+	case ERBFSolverType::Interpolative:
+		SolveInterpolative(Params, Targets, Input, AllWeights);
+		break;
+	}
+
+	float TotalWeight = 0.f; // Keep track of total weight generated, to normalize at the end
+	for (float Weight : AllWeights)
+	{
 		TotalWeight += Weight;
 	}
 
@@ -296,5 +460,11 @@ bool FRBFSolver::FindTargetNeighbourDistances(const FRBFParams& Params, const TA
 
 float FRBFSolver::GetRadiusForTarget(const FRBFTarget& Target, const FRBFParams& Params)
 {
-	return FMath::Max(Params.Radius * Target.ScaleFactor, KINDA_SMALL_NUMBER);
+	float Radius = Params.Radius;
+	if (Params.SolverType == ERBFSolverType::Additive)
+	{
+		Radius *= Target.ScaleFactor;
+	}
+
+	return FMath::Max(Radius, KINDA_SMALL_NUMBER);
 }
diff --git a/Engine/Source/Runtime/AnimGraphRuntime/Public/AnimNodes/AnimNode_LayeredBoneBlend.h b/Engine/Source/Runtime/AnimGraphRuntime/Public/AnimNodes/AnimNode_LayeredBoneBlend.h
--- a/Engine/Source/Runtime/AnimGraphRuntime/Public/AnimNodes/AnimNode_LayeredBoneBlend.h
+++ b/Engine/Source/Runtime/AnimGraphRuntime/Public/AnimNodes/AnimNode_LayeredBoneBlend.h
@@ -87,7 +87,7 @@ public:
 	FAnimNode_LayeredBoneBlend()
 		: bMeshSpaceRotationBlend(false)
 		, bMeshSpaceScaleBlend(false)
-		, CurveBlendOption(ECurveBlendOption::MaxWeight)
+		, CurveBlendOption(ECurveBlendOption::Override)
 		, bBlendRootMotionBasedOnRootBone(true)
 		, bHasRelevantPoses(false)
 		, LODThreshold(INDEX_NONE)
diff --git a/Engine/Source/Runtime/AnimGraphRuntime/Public/AnimNodes/AnimNode_RandomPlayer.h b/Engine/Source/Runtime/AnimGraphRuntime/Public/AnimNodes/AnimNode_RandomPlayer.h
--- a/Engine/Source/Runtime/AnimGraphRuntime/Public/AnimNodes/AnimNode_RandomPlayer.h
+++ b/Engine/Source/Runtime/AnimGraphRuntime/Public/AnimNodes/AnimNode_RandomPlayer.h
@@ -2,13 +2,13 @@
 
 #pragma once
 
-#include "CoreMinimal.h"
-#include "UObject/ObjectMacros.h"
-#include "Math/RandomStream.h"
-#include "Animation/AnimationAsset.h"
-#include "Animation/AnimSequenceDecompressionContext.h"
 #include "AlphaBlend.h"
 #include "Animation/AnimNodeBase.h"
+#include "Animation/AnimationAsset.h"
+#include "CoreMinimal.h"
+#include "Math/RandomStream.h"
+#include "UObject/ObjectMacros.h"
+
 #include "AnimNode_RandomPlayer.generated.h"
 
 enum class ERandomDataIndexType
@@ -17,37 +17,6 @@ enum class ERandomDataIndexType
 	Next,
 };
 
-struct FRandomAnimPlayData
-{
-	FRandomAnimPlayData()
-		: PreviousTimeAccumulator(0.0f)
-		, InternalTimeAccumulator(0.0f)
-		, PlayRate(1.0f)
-		, BlendWeight(0.0f)
-		, RemainingLoops(0)
-	{
-
-	}
-
-	// Prev time, used to track loops (prev > current)
-	float PreviousTimeAccumulator;
-
-	// Current time through the sequence
-	float InternalTimeAccumulator;
-
-	// Calculated play rate 
-	float PlayRate;
-
-	// Current blend weight
-	float BlendWeight;
-
-	// Calculated loops remaining
-	int32 RemainingLoops;
-
-	// Marker tick record for this play through
-	FMarkerTickRecord MarkerTickRecord;
-};
-
 /** The random player node holds a list of sequences and parameter ranges which will be played continuously
   * In a random order. If shuffle mode is enabled then each entry will be played once before repeating any
   */
@@ -57,14 +26,13 @@ struct FRandomPlayerSequenceEntry
 	GENERATED_BODY()
 
 	FRandomPlayerSequenceEntry()
-	: Sequence(nullptr)
-	, ChanceToPlay(1.0f)
-	, MinLoopCount(0)
-	, MaxLoopCount(0)
-	, MinPlayRate(1.0f)
-	, MaxPlayRate(1.0f)
+	    : Sequence(nullptr)
+	    , ChanceToPlay(1.0f)
+	    , MinLoopCount(0)
+	    , MaxLoopCount(0)
+	    , MinPlayRate(1.0f)
+	    , MaxPlayRate(1.0f)
 	{
-
 	}
 
 	/** Sequence to play when this entry is picked */
@@ -72,23 +40,23 @@ struct FRandomPlayerSequenceEntry
 	UAnimSequence* Sequence;
 
 	/** When not in shuffle mode, this is the chance this entry will play (normalized against all other sample chances) */
-	UPROPERTY(EditAnywhere, BlueprintReadWrite, Category = "Settings")
+	UPROPERTY(EditAnywhere, BlueprintReadWrite, Category = "Settings", meta = (UIMin = "0", ClampMin = "0"))
 	float ChanceToPlay;
 
 	/** Minimum number of times this entry will loop before ending */
-	UPROPERTY(EditAnywhere, BlueprintReadWrite, Category = "Settings", meta=(UIMin="0", ClampMin="0"))
+	UPROPERTY(EditAnywhere, BlueprintReadWrite, Category = "Settings", meta = (UIMin = "0", ClampMin = "0"))
 	int32 MinLoopCount;
 
 	/** Maximum number of times this entry will loop before ending */
-	UPROPERTY(EditAnywhere, BlueprintReadWrite, Category = "Settings", meta=(UIMin="0", ClampMin="0"))
+	UPROPERTY(EditAnywhere, BlueprintReadWrite, Category = "Settings", meta = (UIMin = "0", ClampMin = "0"))
 	int32 MaxLoopCount;
 
 	/** Minimum playrate for this entry */
-	UPROPERTY(EditAnywhere, BlueprintReadWrite, Category = "Settings", meta=(UIMin="0", ClampMin="0"))
+	UPROPERTY(EditAnywhere, BlueprintReadWrite, Category = "Settings", meta = (UIMin = "0", ClampMin = "0"))
 	float MinPlayRate;
 
 	/** Maximum playrate for this entry */
-	UPROPERTY(EditAnywhere, BlueprintReadWrite, Category = "Settings", meta=(UIMin="0", ClampMin="0"))
+	UPROPERTY(EditAnywhere, BlueprintReadWrite, Category = "Settings", meta = (UIMin = "0", ClampMin = "0"))
 	float MaxPlayRate;
 
 	/** Blending properties used when this entry is blending in ontop of another entry */
@@ -96,6 +64,34 @@ struct FRandomPlayerSequenceEntry
 	FAlphaBlend BlendIn;
 };
 
+struct FRandomAnimPlayData
+{
+	// Index into the real sequence entry list, not the valid entry list.
+	FRandomPlayerSequenceEntry* Entry = nullptr;
+
+	// The time at which the animation started playing. Used to initialize
+	// the play for this animation and detect when a loop has occurred.
+	float PlayStartTime = 0.0f;
+
+	// The time at which the animation is currently playing.
+	float CurrentPlayTime = 0.0f;
+
+	// The previous tick's update time.
+	float PreviousPlayTime = 0.0f;
+
+	// Calculated play rate
+	float PlayRate = 0.0f;
+
+	// Current blend weight
+	float BlendWeight = 0.0f;
+
+	// Calculated loops remaining
+	int32 RemainingLoops = 0;
+
+	// Marker tick record for this play through
+	FMarkerTickRecord MarkerTickRecord;
+};
+
 USTRUCT(BlueprintInternalUseOnly)
 struct ANIMGRAPHRUNTIME_API FAnimNode_RandomPlayer : public FAnimNode_Base
 {
@@ -116,13 +112,30 @@ public:
 	// End of FAnimNode_Base interface
 
 private:
+	// Return the index of the next FRandomPlayerSequenceEntry to play, from the list
+	// of valid playable entries (ValidEntries).
+	int32 GetNextValidEntryIndex();
 
-	int32 GetNextEntryIndex();
-	int32 GetDataIndex(const ERandomDataIndexType& Type);
+	// Return the play data for either the currently playing animation or the next
+	// animation to blend into.
+	FRandomAnimPlayData& GetPlayData(ERandomDataIndexType Type);
 
-	void SwitchNextToCurrent();
+	// Initialize the play data with the given index into the ValidEntries array and
+	// a specific blend weight. All other member data will be reset to their default values.
+	void InitPlayData(FRandomAnimPlayData& Data, int32 ValidEntryIndex, float BlendWeight);
 
-	void BuildShuffleList();
+	// Advance to the next playable sequence. This is only called once a sequence is fully
+	// blended or there's a hard switch to the same playable entry.
+	void AdvanceToNextSequence();
+
+	// Build a new ShuffleList array, which is a shuffled index list of all the valid
+	// playable entries in ValidEntries. The LastEntry can be set to a valid entry index to
+	// ensure that the top/last item in the shuffle list will be a different value from it;
+	// pass in INDEX_NONE to disable the check.
+	void BuildShuffleList(int32 LastEntry);
+
+	// List of valid sequence entries
+	TArray<FRandomPlayerSequenceEntry*> ValidEntries;
 
 	// Normalized list of play chances when we aren't using shuffle mode
 	TArray<float> NormalizedPlayChances;
@@ -130,19 +143,12 @@ private:
 	// Play data for the current and next sequence
 	TArray<FRandomAnimPlayData> PlayData;
 
+	// Index of the 'current' data set in the PlayData array.
+	int32 CurrentPlayDataIndex;
+
 	// List to store transient shuffle stack in shuffle mode.
 	TArray<int32> ShuffleList;
 
-	// The currently playing entry in the entries list
-	int32 CurrentEntry;
-
-	// The next entry to play, we need this early so we can handle the crossfade correctly as entries
-	// can all have different blend in times.
-	int32 NextEntry;
-
-	// Index of the 'current' data set in the PlayData array
-	int32 CurrentDataIndex;
-
 	// Random number source
 	FRandomStream RandomStream;
 
diff --git a/Engine/Source/Runtime/AnimGraphRuntime/Public/AnimNodes/AnimNode_SequenceEvaluator.h b/Engine/Source/Runtime/AnimGraphRuntime/Public/AnimNodes/AnimNode_SequenceEvaluator.h
--- a/Engine/Source/Runtime/AnimGraphRuntime/Public/AnimNodes/AnimNode_SequenceEvaluator.h
+++ b/Engine/Source/Runtime/AnimGraphRuntime/Public/AnimNodes/AnimNode_SequenceEvaluator.h
@@ -6,7 +6,6 @@
 #include "UObject/ObjectMacros.h"
 #include "Animation/AnimNode_AssetPlayerBase.h"
 #include "Animation/AnimSequenceBase.h"
-#include "Animation/AnimSequenceDecompressionContext.h"
 #include "AnimNode_SequenceEvaluator.generated.h"
 
 UENUM(BlueprintType)
diff --git a/Engine/Source/Runtime/AnimGraphRuntime/Public/AnimNodes/AnimNode_TwoWayBlend.h b/Engine/Source/Runtime/AnimGraphRuntime/Public/AnimNodes/AnimNode_TwoWayBlend.h
--- a/Engine/Source/Runtime/AnimGraphRuntime/Public/AnimNodes/AnimNode_TwoWayBlend.h
+++ b/Engine/Source/Runtime/AnimGraphRuntime/Public/AnimNodes/AnimNode_TwoWayBlend.h
@@ -21,9 +21,13 @@ public:
 	UPROPERTY(EditAnywhere, BlueprintReadWrite, Category=Links)
 	FPoseLink B;
 
+	/** The data type used to control the alpha blending between the A and B poses. 
+		Note: Changing this value will disconnect alpha input pins. 
+	*/
 	UPROPERTY(EditAnywhere, BlueprintReadWrite, Category = Settings)
 	EAnimAlphaInputType AlphaInputType;
 
+	/** The boolean value that controls the alpha blending when the alpha input type is set to 'Bool' */
 	UPROPERTY(EditAnywhere, BlueprintReadWrite, Category = Settings, meta = (PinShownByDefault, DisplayName = "bEnabled", DisplayAfter="AlphaScaleBias"))
 	uint8 bAlphaBoolEnabled:1;
 
@@ -37,6 +41,7 @@ protected:
 	uint8 bResetChildOnActivation:1;
 
 public:
+	/** The float value that controls the alpha blending when the alpha input type is set to 'Float' */
 	UPROPERTY(EditAnywhere, BlueprintReadWrite, Category=Settings, meta=(PinShownByDefault))
 	float Alpha;
 
@@ -46,6 +51,7 @@ public:
 	UPROPERTY(EditAnywhere, BlueprintReadWrite, Category = Settings, meta = (DisplayName = "Blend Settings"))
 	FInputAlphaBoolBlend AlphaBoolBlend;
 
+	/** The animation curve that controls the alpha blending when the alpha input type is set to 'Curve' */
 	UPROPERTY(EditAnywhere, BlueprintReadWrite, Category = Settings, meta = (PinShownByDefault))
 	FName AlphaCurveName;
 
diff --git a/Engine/Source/Runtime/AnimGraphRuntime/Public/AnimSequencerInstance.h b/Engine/Source/Runtime/AnimGraphRuntime/Public/AnimSequencerInstance.h
--- a/Engine/Source/Runtime/AnimGraphRuntime/Public/AnimSequencerInstance.h
+++ b/Engine/Source/Runtime/AnimGraphRuntime/Public/AnimSequencerInstance.h
@@ -8,29 +8,32 @@
 
 #pragma once
 #include "Animation/AnimInstance.h"
-#include "AnimCustomInstance.h"
+#include "SequencerAnimationSupport.h"
 #include "AnimSequencerInstance.generated.h"
 
 UCLASS(transient, NotBlueprintable)
-class ANIMGRAPHRUNTIME_API UAnimSequencerInstance : public UAnimCustomInstance
+class ANIMGRAPHRUNTIME_API UAnimSequencerInstance : public UAnimInstance, public ISequencerAnimationSupport
 {
 	GENERATED_UCLASS_BODY()
 
 public:
 
 	/** Update an animation sequence player in this instance */
-	void UpdateAnimTrack(UAnimSequenceBase* InAnimSequence, int32 SequenceId, float InPosition, float Weight, bool bFireNotifies);
-	void UpdateAnimTrack(UAnimSequenceBase* InAnimSequence, int32 SequenceId, float InFromPosition, float InToPosition, float Weight, bool bFireNotifies);
+	virtual void UpdateAnimTrack(UAnimSequenceBase* InAnimSequence, int32 SequenceId, float InPosition, float Weight, bool bFireNotifies);
+	virtual void UpdateAnimTrack(UAnimSequenceBase* InAnimSequence, int32 SequenceId, float InFromPosition, float InToPosition, float Weight, bool bFireNotifies);
 
 	/** Reset all nodes in this instance */
-	void ResetNodes();
+	virtual void ResetNodes() override;
 
 	/** Reset the pose for this instance*/
-	void ResetPose();
+	virtual void ResetPose() override;
 
 	/** Saved the named pose to restore after */
-	void SavePose();
+	virtual void SavePose() override;
 
+	virtual UAnimInstance* GetSourceAnimInstance() override { return this; }
+	virtual void SetSourceAnimInstance(UAnimInstance* SourceAnimInstance) {  /* nothing to do */ ensure(false); }
+	virtual bool DoesSupportDifferentSourceAnimInstance() const override { return false; }
 protected:
 	// UAnimInstance interface
 	virtual FAnimInstanceProxy* CreateAnimInstanceProxy() override;
@@ -41,3 +44,4 @@ public:
 	static const FName SequencerPoseName;
 
 };
+
diff --git a/Engine/Source/Runtime/AnimGraphRuntime/Public/RBF/RBFSolver.h b/Engine/Source/Runtime/AnimGraphRuntime/Public/RBF/RBFSolver.h
--- a/Engine/Source/Runtime/AnimGraphRuntime/Public/RBF/RBFSolver.h
+++ b/Engine/Source/Runtime/AnimGraphRuntime/Public/RBF/RBFSolver.h
@@ -8,6 +8,27 @@
 #include "Curves/RichCurve.h"
 #include "RBFSolver.generated.h"
 
+
+/** The solver type to use. The two solvers have different requirements. */
+UENUM()
+enum class ERBFSolverType : uint8
+{
+	/** The additive solver sums up contributions from each target. It's faster
+	    but may require more targets for a good coverage, and requires the 
+		normalization step to be performed for smooth results.
+	*/
+	Additive,
+
+	/** The interpolative solver interpolates the values from each target based
+		on distance. As long as the input values are within the area bounded by
+		the targets, the interpolation is well-behaved and return weight values 
+		within the 0% - 100% limit with no normalization required. 
+		Interpolation also gives smoother results, with fewer targets, than additive
+		but at a higher computational cost.
+	*/
+	Interpolative
+};
+
 /** Function to use for each target falloff */
 UENUM()
 enum class ERBFFunctionType : uint8
@@ -39,6 +60,9 @@ enum class ERBFDistanceMethod : uint8
 	/** Treat inputs as quaternion, and find distance between rotated TwistAxis direction */
 	SwingAngle,
 
+	/** Treat inputs as quaternion, and find distance between rotations around the TwistAxis direction */
+	TwistAngle,
+
 	/** Uses the setting of the parent container */
 	DefaultMethod
 };
@@ -63,7 +87,15 @@ enum class ERBFNormalizeMethod : uint8
 		non-normalized and normalized. This helps to define
 		the volume in which normalization is always required.
 	*/
-	NormalizeWithinMedian
+	NormalizeWithinMedian,
+
+	/** 
+		Don't normalize at all. This should only be used with
+		the interpolative method, if it is known that all input
+		values will be within the area bounded by the targets.
+	*/
+	NoNormalization,
+
 };
 
 /** Struct storing a particular entry within the RBF */
@@ -76,8 +108,15 @@ struct ANIMGRAPHRUNTIME_API FRBFEntry
 	UPROPERTY(EditAnywhere, Category = RBFData)
 	TArray<float> Values;
 
+	/** Return a target as an rotator, assuming Values is a sequence of Euler entries. Index is which Euler to convert.*/
+	FRotator AsRotator(int32 Index) const;
+
 	/** Return a target as a quaternion, assuming Values is a sequence of Euler entries. Index is which Euler to convert. */
 	FQuat AsQuat(int32 Index) const;
+
+	FVector AsVector(int32 Index) const;
+
+
 	/** Set this entry to 3 floats from supplied rotator */
 	void AddFromRotator(const FRotator& InRot);
 	/** Set this entry to 3 floats from supplied vector */
@@ -96,21 +135,28 @@ struct ANIMGRAPHRUNTIME_API FRBFTarget : public FRBFEntry
 {
 	GENERATED_BODY()
 
-	/** How large to scale */
+	/** How large the influence of this target. */
 	UPROPERTY(EditAnywhere, Category = RBFData)
 	float ScaleFactor;
 
-	/** Whether we want to apply an additional custom curve when activating this target */
+	/** Whether we want to apply an additional custom curve when activating this target. 
+	    Ignored if the solver type is Interpolative. 
+	*/
 	UPROPERTY(EditAnywhere, Category = RBFData)
 	bool bApplyCustomCurve;
 
-	/** Custom curve to apply to activation of this target, if bApplyCustomCurve is true */
+	/** Custom curve to apply to activation of this target, if bApplyCustomCurve is true.
+		Ignored if the solver type is Interpolative. */
 	UPROPERTY(EditAnywhere, Category = RBFData)
 	FRichCurve CustomCurve;
 
+	/** Override the distance method used to calculate the distance from this target to
+		the input. Ignored if the solver type is Interpolative. */
 	UPROPERTY(EditAnywhere, Category = RBFData)
 	ERBFDistanceMethod DistanceMethod;
 
+	/** Override the falloff function used to smooth the distance from this target to
+		the input. Ignored if the solver type is Interpolative. */
 	UPROPERTY(EditAnywhere, Category = RBFData)
 	ERBFFunctionType FunctionType;
 
@@ -151,7 +197,16 @@ struct ANIMGRAPHRUNTIME_API FRBFParams
 	UPROPERTY()
 	int32 TargetDimensions;
 
-	/** Default radius for each target, scaled by per-Target ScaleFactor */
+	/** Specifies the type of solver to use. The additive solver requires normalization, for the
+		most part, whereas the Interpolative solver is not as reliant on it. The interpolative
+		solver also has smoother blending, whereas the additive solver requires more targets but
+		has a more precise control over the influence of each target.
+	*/
+	UPROPERTY(EditAnywhere, Category = RBFData)
+	ERBFSolverType SolverType;
+
+	/** Default radius for each target. 
+	*/
 	UPROPERTY(EditAnywhere, Category = RBFData)
 	float Radius;
 
@@ -194,6 +249,10 @@ struct ANIMGRAPHRUNTIME_API FRBFParams
 /** Library of Radial Basis Function solver functions */
 struct ANIMGRAPHRUNTIME_API FRBFSolver
 {
+	/** Given a list of targets, verify which ones are valid for solving the RBF setup. This is mostly about removing identical targets
+		which invalidates the interpolative solver. Returns true if all targets are valid. */
+	static bool ValidateTargets(const FRBFParams& Params, const TArray<FRBFTarget>& Targets, TArray<int>& InvalidTargets);
+
 	/** Given a set of targets and new input entry, give list of activated targets with weights */
 	static void Solve(const FRBFParams& Params, const TArray<FRBFTarget>& Targets, const FRBFEntry& Input, TArray<FRBFOutputWeight>& OutputWeights);
 
diff --git a/Engine/Source/Runtime/AnimationCore/Private/AnimationDataSource.cpp b/Engine/Source/Runtime/AnimationCore/Private/AnimationDataSource.cpp
--- a/Engine/Source/Runtime/AnimationCore/Private/AnimationDataSource.cpp
+++ b/Engine/Source/Runtime/AnimationCore/Private/AnimationDataSource.cpp
@@ -4,6 +4,8 @@
 
 bool UAnimationDataSourceRegistry::RegisterDataSource(const FName& InName, UObject* InDataSource)
 {
+	ClearInvalidDataSource();
+
 	if (DataSources.Contains(InName))
 	{
 		return false;
@@ -14,6 +16,7 @@ bool UAnimationDataSourceRegistry::RegisterDataSource(const FName& InName, UObje
 
 bool UAnimationDataSourceRegistry::UnregisterDataSource(const FName& InName)
 {
+	ClearInvalidDataSource();
 	return DataSources.Remove(InName) > 0;
 }
 
@@ -24,18 +27,38 @@ bool UAnimationDataSourceRegistry::ContainsSource(const FName& InName) const
 
 UObject* UAnimationDataSourceRegistry::RequestSource(const FName& InName, UClass* InExpectedClass) const
 {
-	UObject* const* DataSource = DataSources.Find(InName);
+	TWeakObjectPtr<UObject> const* DataSource = DataSources.Find(InName);
 	if (DataSource == nullptr)
 	{
 		return nullptr;
 	}
-	if (*DataSource == nullptr)
+
+	UObject* DataSourcePtr = DataSource->Get();
+	if (!DataSourcePtr)
 	{
 		return nullptr;
 	}
-	if (!(*DataSource)->IsA(InExpectedClass))
+	if (!DataSourcePtr->IsA(InExpectedClass))
 	{
 		return nullptr;
 	}
-	return *DataSource;
+
+	return DataSourcePtr;
+}
+
+void UAnimationDataSourceRegistry::ClearInvalidDataSource()
+{
+	TArray<FName> InvalidNames;
+	for (auto Iter = DataSources.CreateIterator(); Iter; ++Iter)
+	{
+		if (!Iter.Value().IsValid())
+		{
+			InvalidNames.Add(Iter.Key());
+		}
+	}
+
+	for (const FName& NameToRemove : InvalidNames)
+	{
+		DataSources.Remove(NameToRemove);
+	}
 }
\ No newline at end of file
diff --git a/Engine/Source/Runtime/AnimationCore/Public/AnimationDataSource.h b/Engine/Source/Runtime/AnimationCore/Public/AnimationDataSource.h
--- a/Engine/Source/Runtime/AnimationCore/Public/AnimationDataSource.h
+++ b/Engine/Source/Runtime/AnimationCore/Public/AnimationDataSource.h
@@ -69,6 +69,9 @@ public:
 
 private:
 
-	UPROPERTY()
-	TMap<FName, UObject*> DataSources;
+	UPROPERTY(transient)
+	TMap<FName, TWeakObjectPtr<UObject>> DataSources;
+
+	/** Clear Invalid Data Sources that are GC-ed */
+	void ClearInvalidDataSource();
 };
\ No newline at end of file
diff --git a/Engine/Source/Runtime/AnimationCore/Public/EulerTransform.h b/Engine/Source/Runtime/AnimationCore/Public/EulerTransform.h
--- a/Engine/Source/Runtime/AnimationCore/Public/EulerTransform.h
+++ b/Engine/Source/Runtime/AnimationCore/Public/EulerTransform.h
@@ -15,20 +15,39 @@ struct ANIMATIONCORE_API FEulerTransform
 	 */
 	static const FEulerTransform Identity;
 
-	FEulerTransform()
+	FORCEINLINE FEulerTransform()
 		: Location(ForceInitToZero)
 		, Rotation(ForceInitToZero)
 		, Scale(ForceInitToZero)
 	{
 	}
 
-	FEulerTransform(const FVector& InLocation, const FRotator& InRotation, const FVector& InScale)
+	FORCEINLINE FEulerTransform(const FVector& InLocation, const FRotator& InRotation, const FVector& InScale)
 		: Location(InLocation)
 		, Rotation(InRotation)
 		, Scale(InScale)
 	{
 	}
 
+	FORCEINLINE FEulerTransform(const FTransform& InTransform)
+		: Location(InTransform.GetLocation())
+		, Rotation(InTransform.GetRotation().Rotator())
+		, Scale(InTransform.GetScale3D())
+	{
+
+	}
+
+	FORCEINLINE FEulerTransform& operator =(const FTransform& InTransform)
+	{
+		FromFTransform(InTransform);
+		return *this;
+	}
+
+	FORCEINLINE operator FTransform() const
+	{
+		return ToFTransform();
+	}
+
 	/** The translation of this transform */
 	UPROPERTY(EditAnywhere, BlueprintReadWrite, Category="Transform")
 	FVector Location;
@@ -42,13 +61,13 @@ struct ANIMATIONCORE_API FEulerTransform
 	FVector Scale;
 
 	/** Convert to an FTransform */
-	FTransform ToFTransform() const
+	FORCEINLINE FTransform ToFTransform() const
 	{
 		return FTransform(Rotation.Quaternion(), Location, Scale);
 	}
 
 	/** Convert from an FTransform */
-	void FromFTransform(const FTransform& InTransform)
+	FORCEINLINE void FromFTransform(const FTransform& InTransform)
 	{
 		Location = InTransform.GetLocation();
 		Rotation = InTransform.GetRotation().Rotator();
diff --git a/Engine/Source/Runtime/AnimationCore/Public/TransformNoScale.h b/Engine/Source/Runtime/AnimationCore/Public/TransformNoScale.h
new file mode 100644
--- /dev/null
+++ b/Engine/Source/Runtime/AnimationCore/Public/TransformNoScale.h
@@ -0,0 +1,67 @@
+// Copyright Epic Games, Inc. All Rights Reserved.
+
+#pragma once
+
+#include "CoreMinimal.h"
+#include "TransformNoScale.generated.h"
+
+USTRUCT(BlueprintType)
+struct ANIMATIONCORE_API FTransformNoScale
+{
+	GENERATED_BODY()
+
+	/**
+	 * The identity transformation (Rotation = FRotator::ZeroRotator, Translation = FVector::ZeroVector, Scale = (1,1,1)).
+	 */
+	static const FTransformNoScale Identity;
+
+	FORCEINLINE_DEBUGGABLE FTransformNoScale()
+		: Location(ForceInitToZero)
+		, Rotation(ForceInitToZero)
+	{
+	}
+
+	FORCEINLINE_DEBUGGABLE FTransformNoScale(const FVector& InLocation, const FQuat& InRotation)
+		: Location(InLocation)
+		, Rotation(InRotation)
+	{
+	}
+
+	FORCEINLINE_DEBUGGABLE FTransformNoScale(const FTransform& InTransform)
+		: Location(InTransform.GetLocation())
+		, Rotation(InTransform.GetRotation())
+	{
+	}
+
+	FORCEINLINE_DEBUGGABLE FTransformNoScale& operator =(const FTransform& InTransform)
+	{
+		FromFTransform(InTransform);
+		return *this;
+	}
+
+	FORCEINLINE_DEBUGGABLE operator FTransform() const
+	{
+		return ToFTransform();
+	}
+
+	/** The translation of this transform */
+	UPROPERTY(EditAnywhere, BlueprintReadWrite, Category="Transform")
+	FVector Location;
+
+	/** The rotation of this transform */
+	UPROPERTY(EditAnywhere, BlueprintReadWrite, Category="Transform")
+	FQuat Rotation;
+
+	/** Convert to an FTransform */
+	FORCEINLINE_DEBUGGABLE FTransform ToFTransform() const
+	{
+		return FTransform(Rotation, Location, FVector::OneVector);
+	}
+
+	/** Convert from an FTransform */
+	FORCEINLINE_DEBUGGABLE void FromFTransform(const FTransform& InTransform)
+	{
+		Location = InTransform.GetLocation();
+		Rotation = InTransform.GetRotation();
+	}
+};
\ No newline at end of file
diff --git a/Engine/Source/Runtime/ClothingSystemRuntimeInterface/Public/ClothVertBoneData.h b/Engine/Source/Runtime/ClothingSystemRuntimeInterface/Public/ClothVertBoneData.h
--- a/Engine/Source/Runtime/ClothingSystemRuntimeInterface/Public/ClothVertBoneData.h
+++ b/Engine/Source/Runtime/ClothingSystemRuntimeInterface/Public/ClothVertBoneData.h
@@ -17,11 +17,11 @@ struct FClothVertBoneData
 		FMemory::Memset(BoneWeights, 0, sizeof(BoneWeights));
 	}
 
-	// MAX_TOTAL_INFLUENCES = 8 is defined in GPUSkinPublicDefs.h, but that would 
+	// MAX_TOTAL_INFLUENCES = 12 is defined in GPUSkinPublicDefs.h, but that would 
 	// require a dependency on Engine, which we can't have here.  Until that can
 	// be migrated elsewhere (not to mention if), we make a redundant variable
 	// with an 8 in it, that needs to stay in sync with MAX_TOTAL_INFLUENCES.
-	static const int8 MaxTotalInfluences = 8;
+	static const int8 MaxTotalInfluences = 12;
 
 	// Number of influences for this vertex.
 	UPROPERTY()
