Chaos - fix AnimNode_RigidBody_Chaos to match regular version after iterations changes.

diff --git a/Engine/Source/Runtime/AnimGraphRuntime/Private/BoneControllers/AnimNode_RigidBody_Chaos.cpp b/Engine/Source/Runtime/AnimGraphRuntime/Private/BoneControllers/AnimNode_RigidBody_Chaos.cpp
--- a/Engine/Source/Runtime/AnimGraphRuntime/Private/BoneControllers/AnimNode_RigidBody_Chaos.cpp
+++ b/Engine/Source/Runtime/AnimGraphRuntime/Private/BoneControllers/AnimNode_RigidBody_Chaos.cpp
@@ -52,6 +52,13 @@ FAnimNode_RigidBody_Chaos::FAnimNode_RigidBody_Chaos()
 	bFreezeIncomingPoseOnStart = false;
 	bClampLinearTranslationLimitToRefPose = false;
 
+	OverrideSolverIterations.SolverIterations = -1;
+	OverrideSolverIterations.JointIterations = -1;
+	OverrideSolverIterations.CollisionIterations = -1;
+	OverrideSolverIterations.SolverPushOutIterations = -1;
+	OverrideSolverIterations.JointPushOutIterations = -1;
+	OverrideSolverIterations.CollisionPushOutIterations = -1;
+
 	PreviousTransform = CurrentTransform = FTransform::Identity;
 	PreviousComponentLinearVelocity = FVector::ZeroVector;	
 
@@ -465,6 +472,22 @@ void FAnimNode_RigidBody_Chaos::EvaluateSkeletalControl_AnyThread(FComponentSpac
 			UpdateWorldForces(CompWorldSpaceTM, BaseBoneTM);
 			const FVector SimSpaceGravity = WorldVectorToSpaceNoScaleChaos(SimulationSpace, WorldSpaceGravity, CompWorldSpaceTM, BaseBoneTM);
 
+			PhysicsSimulation->SetSolverIterations(
+				SolverIterations.SolverIterations,
+				SolverIterations.JointIterations,
+				SolverIterations.CollisionIterations,
+				SolverIterations.SolverPushOutIterations,
+				SolverIterations.JointPushOutIterations,
+				SolverIterations.CollisionPushOutIterations
+			);
+			PhysicsSimulation->SetSolverIterations(
+				OverrideSolverIterations.SolverIterations,
+				OverrideSolverIterations.JointIterations,
+				OverrideSolverIterations.CollisionIterations,
+				OverrideSolverIterations.SolverPushOutIterations,
+				OverrideSolverIterations.JointPushOutIterations,
+				OverrideSolverIterations.CollisionPushOutIterations);
+
 			// Run simulation at a minimum of 30 FPS to prevent system from exploding.
 			// DeltaTime can be higher due to URO, so take multiple iterations in that case.
 			extern int32 RBAN_MaxSubSteps;
@@ -796,6 +819,16 @@ void FAnimNode_RigidBody_Chaos::InitPhysics(const UAnimInstance* InAnimInstance)
 
 		PhysicsSimulation->SetIgnoreCollisionPairTable(IgnorePairs);
 		PhysicsSimulation->SetIgnoreCollisionActors(IgnoreCollisionActors);
+
+		SolverIterations = UsePhysicsAsset->SolverIterations;
+		PhysicsSimulation->SetSolverIterations(
+			SolverIterations.SolverIterations,
+			SolverIterations.JointIterations,
+			SolverIterations.CollisionIterations,
+			SolverIterations.SolverPushOutIterations,
+			SolverIterations.JointPushOutIterations,
+			SolverIterations.CollisionPushOutIterations
+		);
 	}
 }
 
diff --git a/Engine/Source/Runtime/AnimGraphRuntime/Public/BoneControllers/AnimNode_RigidBody_Chaos.h b/Engine/Source/Runtime/AnimGraphRuntime/Public/BoneControllers/AnimNode_RigidBody_Chaos.h
--- a/Engine/Source/Runtime/AnimGraphRuntime/Public/BoneControllers/AnimNode_RigidBody_Chaos.h
+++ b/Engine/Source/Runtime/AnimGraphRuntime/Public/BoneControllers/AnimNode_RigidBody_Chaos.h
@@ -137,6 +137,15 @@ public:
 	UPROPERTY(EditAnywhere, Category = Settings)
 	uint8 bClampLinearTranslationLimitToRefPose : 1;
 
+	/**
+	 * Solver iteration settings overrides (defaults are set in the Physics Asset).
+	 * These can be varied in the runtime and set through blueprint (e.g., to increase
+	 * iterations during difficult movements).
+	 * Set to -1 to leave an individual iteration value at its Physics Asset value.
+	 */
+	UPROPERTY(EditAnywhere, Category = Settings, meta = (PinHiddenByDefault))
+		FSolverIterations OverrideSolverIterations;
+
 private:
 #if WITH_EDITORONLY_DATA
 	UPROPERTY()
@@ -172,6 +181,7 @@ private:
 	TWeakObjectPtr<USkeletalMeshComponent> SkelMeshCompWeakPtr;
 
 	ImmediatePhysics_Chaos::FSimulation* PhysicsSimulation;
+	FSolverIterations SolverIterations;
 
 	struct FOutputBoneData
 	{
