RBAN - speculative fix for possible null world

diff --git a/Engine/Source/Runtime/AnimGraphRuntime/Private/BoneControllers/AnimNode_RigidBody.cpp b/Engine/Source/Runtime/AnimGraphRuntime/Private/BoneControllers/AnimNode_RigidBody.cpp
--- a/Engine/Source/Runtime/AnimGraphRuntime/Private/BoneControllers/AnimNode_RigidBody.cpp
+++ b/Engine/Source/Runtime/AnimGraphRuntime/Private/BoneControllers/AnimNode_RigidBody.cpp
@@ -979,7 +979,6 @@ void FAnimNode_RigidBody::PreUpdate(const UAnimInstance* InAnimInstance)
 
 	SCOPE_CYCLE_COUNTER(STAT_RigidBody_PreUpdate);
 
-	UWorld* World = InAnimInstance->GetWorld();
 	USkeletalMeshComponent* SKC = InAnimInstance->GetSkelMeshComponent();
 	APawn* PawnOwner = InAnimInstance->TryGetPawnOwner();
 	UPawnMovementComponent* MovementComp = PawnOwner ? PawnOwner->GetMovementComponent() : nullptr;
@@ -992,19 +991,23 @@ void FAnimNode_RigidBody::PreUpdate(const UAnimInstance* InAnimInstance)
 	}
 #endif
 
-	WorldSpaceGravity = bOverrideWorldGravity ? OverrideWorldGravity : (MovementComp ? FVector(0.f, 0.f, MovementComp->GetGravityZ()) : FVector(0.f, 0.f, World->GetGravityZ()));
-	if(SKC)
+	UWorld* World = InAnimInstance->GetWorld();
+	if (World)
 	{
-		if (PhysicsSimulation && bEnableWorldGeometry && SimulationSpace == ESimulationSpace::WorldSpace)
+		WorldSpaceGravity = bOverrideWorldGravity ? OverrideWorldGravity : (MovementComp ? FVector(0.f, 0.f, MovementComp->GetGravityZ()) : FVector(0.f, 0.f, World->GetGravityZ()));
+		if(SKC)
 		{
-			UpdateWorldGeometry(*World, *SKC);
-		}
+			if (PhysicsSimulation && bEnableWorldGeometry && SimulationSpace == ESimulationSpace::WorldSpace)
+			{
+				UpdateWorldGeometry(*World, *SKC);
+			}
 
-		PendingRadialForces = SKC->GetPendingRadialForces();
+			PendingRadialForces = SKC->GetPendingRadialForces();
 
-		PreviousTransform = CurrentTransform;
-		CurrentTransform = SKC->GetComponentToWorld();
-	}	
+			PreviousTransform = CurrentTransform;
+			CurrentTransform = SKC->GetComponentToWorld();
+		}	
+	}
 }
 
 int32 FAnimNode_RigidBody::GetLODThreshold() const
