RBAN - fix skeleton mismatches and CollisionDisableTable issues

diff --git a/Engine/Source/Runtime/AnimGraphRuntime/Private/BoneControllers/AnimNode_RigidBody.cpp b/Engine/Source/Runtime/AnimGraphRuntime/Private/BoneControllers/AnimNode_RigidBody.cpp
--- a/Engine/Source/Runtime/AnimGraphRuntime/Private/BoneControllers/AnimNode_RigidBody.cpp
+++ b/Engine/Source/Runtime/AnimGraphRuntime/Private/BoneControllers/AnimNode_RigidBody.cpp
@@ -14,7 +14,7 @@
 #include "PhysicsEngine/PhysicsSettings.h"
 #include "Logging/MessageLog.h"
 
-//#pragma optimize("", off)
+//PRAGMA_DISABLE_OPTIMIZATION
 
 /////////////////////////////////////////////////////
 // FAnimNode_RigidBody
@@ -707,8 +707,11 @@ void FAnimNode_RigidBody::InitPhysics(const UAnimInstance* InAnimInstance)
 
 		const int32 NumBonesLOD0 = InsertionOrder.Num();
 
+		// If our skeleton is not the one that was used to build the PhysicsAsset, some bodies may be missing, or rearranged.
+		// We need to map the original indices to the new bodies for use by the CollisionDisableTable.
+		// NOTE: This array is indexed by the original BodyInstance body index (BodyInstance->InstanceBodyIndex)
 		TArray<ImmediatePhysics::FActorHandle*> BodyIndexToActorHandle;
-		BodyIndexToActorHandle.AddZeroed(NumBonesLOD0);
+		BodyIndexToActorHandle.AddZeroed(HighLevelBodyInstances.Num());
 
 		TArray<FBodyInstance*> BodiesSorted;
 		BodiesSorted.AddZeroed(NumBonesLOD0);
@@ -852,10 +855,18 @@ void FAnimNode_RigidBody::InitPhysics(const UAnimInstance* InAnimInstance)
 		const TMap<FRigidBodyIndexPair, bool>& DisableTable = UsePhysicsAsset->CollisionDisableTable;
 		for(auto ConstItr = DisableTable.CreateConstIterator(); ConstItr; ++ConstItr)
 		{
-			ImmediatePhysics::FSimulation::FIgnorePair Pair;
-			Pair.A = BodyIndexToActorHandle[ConstItr.Key().Indices[0]];
-			Pair.B = BodyIndexToActorHandle[ConstItr.Key().Indices[1]];
-			IgnorePairs.Add(Pair);
+			int32 IndexA = ConstItr.Key().Indices[0];
+			int32 IndexB = ConstItr.Key().Indices[1];
+			if (ensure((IndexA < BodyIndexToActorHandle.Num()) && (IndexB < BodyIndexToActorHandle.Num())))
+			{
+				if ((BodyIndexToActorHandle[IndexA] != nullptr) && (BodyIndexToActorHandle[IndexB] != nullptr))
+				{
+					ImmediatePhysics::FSimulation::FIgnorePair Pair;
+					Pair.A = BodyIndexToActorHandle[IndexA];
+					Pair.B = BodyIndexToActorHandle[IndexB];
+					IgnorePairs.Add(Pair);
+				}
+			}
 		}
 
 		PhysicsSimulation->SetIgnoreCollisionPairTable(IgnorePairs);
