RBAN - remove ensure when physics asset has invalid collision disable table

diff --git a/Engine/Source/Runtime/AnimGraphRuntime/Private/BoneControllers/AnimNode_RigidBody.cpp b/Engine/Source/Runtime/AnimGraphRuntime/Private/BoneControllers/AnimNode_RigidBody.cpp
--- a/Engine/Source/Runtime/AnimGraphRuntime/Private/BoneControllers/AnimNode_RigidBody.cpp
+++ b/Engine/Source/Runtime/AnimGraphRuntime/Private/BoneControllers/AnimNode_RigidBody.cpp
@@ -857,7 +857,7 @@ void FAnimNode_RigidBody::InitPhysics(const UAnimInstance* InAnimInstance)
 		{
 			int32 IndexA = ConstItr.Key().Indices[0];
 			int32 IndexB = ConstItr.Key().Indices[1];
-			if (ensure((IndexA < BodyIndexToActorHandle.Num()) && (IndexB < BodyIndexToActorHandle.Num())))
+			if ((IndexA < BodyIndexToActorHandle.Num()) && (IndexB < BodyIndexToActorHandle.Num()))
 			{
 				if ((BodyIndexToActorHandle[IndexA] != nullptr) && (BodyIndexToActorHandle[IndexB] != nullptr))
 				{
