Guard against null bodysetup in AnimNode to fix crash in ChecKRBN engine test.

diff --git a/Engine/Source/Runtime/AnimGraphRuntime/Private/BoneControllers/AnimNode_RigidBody.cpp b/Engine/Source/Runtime/AnimGraphRuntime/Private/BoneControllers/AnimNode_RigidBody.cpp
--- a/Engine/Source/Runtime/AnimGraphRuntime/Private/BoneControllers/AnimNode_RigidBody.cpp
+++ b/Engine/Source/Runtime/AnimGraphRuntime/Private/BoneControllers/AnimNode_RigidBody.cpp
@@ -1076,7 +1076,12 @@ void FAnimNode_RigidBody::UpdateInternal(const FAnimationUpdateContext& Context)
 				if (ComponentsInSim.Contains(OverlapComp) == false)
 				{
 					ComponentsInSim.Add(OverlapComp);
-					PhysicsSimulation->CreateActor(ImmediatePhysics::EActorType::StaticActor, &OverlapComp->BodyInstance, OverlapComp->BodyInstance.GetUnrealWorldTransform());
+
+					// Not sure why this happens, adding check to fix crash in CheckRBN engine test.
+					if (OverlapComp->BodyInstance.BodySetup != nullptr)
+					{
+						PhysicsSimulation->CreateActor(ImmediatePhysics::EActorType::StaticActor, &OverlapComp->BodyInstance, OverlapComp->BodyInstance.GetUnrealWorldTransform());
+					}
 				}
 			}
 		}
