RBAN - fix initialization crash when skeleton LOD0 ignores some bones in middle of hierarchy

diff --git a/Engine/Source/Runtime/AnimGraphRuntime/Private/BoneControllers/AnimNode_RigidBody.cpp b/Engine/Source/Runtime/AnimGraphRuntime/Private/BoneControllers/AnimNode_RigidBody.cpp
--- a/Engine/Source/Runtime/AnimGraphRuntime/Private/BoneControllers/AnimNode_RigidBody.cpp
+++ b/Engine/Source/Runtime/AnimGraphRuntime/Private/BoneControllers/AnimNode_RigidBody.cpp
@@ -597,16 +597,20 @@ void ComputeBodyInsertionOrder(TArray<FBoneIndexType>& InsertionOrder, const USk
 
 	InsertionOrder.Reset();
 
+	if (SKC.SkeletalMesh == nullptr)
+	{
+		return;
+	}
+
 	const int32 NumLODs = SKC.GetNumLODs();
 	if(NumLODs > 0)
 	{
-		TArray<bool> InSortedOrder;
-
 		TArray<FBoneIndexType> RequiredBones0;
 		TArray<FBoneIndexType> ComponentSpaceTMs0;
 		SKC.ComputeRequiredBones(RequiredBones0, ComponentSpaceTMs0, 0, /*bIgnorePhysicsAsset=*/ true);
 
-		InSortedOrder.AddZeroed(RequiredBones0.Num());
+		TArray<bool> InSortedOrder;
+		InSortedOrder.AddZeroed(SKC.SkeletalMesh->RefSkeleton.GetNum());
 
 		auto MergeIndices = [&InsertionOrder, &InSortedOrder](const TArray<FBoneIndexType>& RequiredBones) -> void
 		{
@@ -707,7 +711,11 @@ void FAnimNode_RigidBody::InitPhysics(const UAnimInstance* InAnimInstance)
 		TArray<FBoneIndexType> InsertionOrder;
 		ComputeBodyInsertionOrder(InsertionOrder, *SkeletalMeshComp);
 
+		// NOTE: NumBonesLOD0 may be less than NumBonesTotal, and it may be middle bones that are missing from LOD0.
+		// In this case, LOD0 bone indices may be >= NumBonesLOD0, but always < NumBonesTotal. Arrays indexed by
+		// bone index must be size NumBonesTotal.
 		const int32 NumBonesLOD0 = InsertionOrder.Num();
+		const int32 NumBonesTotal = SkelMeshRefSkel.GetNum();
 
 		// If our skeleton is not the one that was used to build the PhysicsAsset, some bodies may be missing, or rearranged.
 		// We need to map the original indices to the new bodies for use by the CollisionDisableTable.
@@ -716,7 +724,7 @@ void FAnimNode_RigidBody::InitPhysics(const UAnimInstance* InAnimInstance)
 		BodyIndexToActorHandle.AddZeroed(HighLevelBodyInstances.Num());
 
 		TArray<FBodyInstance*> BodiesSorted;
-		BodiesSorted.AddZeroed(NumBonesLOD0);
+		BodiesSorted.AddZeroed(NumBonesTotal);
 
 		for (FBodyInstance* BI : HighLevelBodyInstances)
 		{
@@ -765,7 +773,7 @@ void FAnimNode_RigidBody::InitPhysics(const UAnimInstance* InAnimInstance)
 		//Insert joints so that they coincide body order. That is, if we stop simulating all bodies past some index, we can simply ignore joints past a corresponding index without any re-order
 		//For this to work we consider the most last inserted bone in each joint
 		TArray<int32> InsertionOrderPerBone;
-		InsertionOrderPerBone.AddUninitialized(NumBonesLOD0);
+		InsertionOrderPerBone.AddUninitialized(NumBonesTotal);
 
 		for(int32 Position = 0; Position < NumBonesLOD0; ++Position)
 		{
