Chaos - RBAN - support "Ignore Collision" on constraints in PhAT

diff --git a/Engine/Source/Runtime/AnimGraphRuntime/Private/BoneControllers/AnimNode_RigidBody.cpp b/Engine/Source/Runtime/AnimGraphRuntime/Private/BoneControllers/AnimNode_RigidBody.cpp
--- a/Engine/Source/Runtime/AnimGraphRuntime/Private/BoneControllers/AnimNode_RigidBody.cpp
+++ b/Engine/Source/Runtime/AnimGraphRuntime/Private/BoneControllers/AnimNode_RigidBody.cpp
@@ -800,6 +800,7 @@ void FAnimNode_RigidBody::InitPhysics(const UAnimInstance* InAnimInstance)
 		});
 
 
+		TArray<ImmediatePhysics::FSimulation::FIgnorePair> IgnorePairs;
 		if(NamesToHandles.Num() > 0)
 		{
 			//constraints
@@ -837,6 +838,14 @@ void FAnimNode_RigidBody::InitPhysics(const UAnimInstance* InAnimInstance)
 							FTransform Body2Transform = Body2Handle->GetWorldTransform();
 							BodyAnimData[BodyIndex].RefPoseLength = Body1Transform.GetRelativeTransform(Body2Transform).GetLocation().Size();
 						}
+
+						if (CI->IsCollisionDisabled())
+						{
+							ImmediatePhysics::FSimulation::FIgnorePair Pair;
+							Pair.A = Body1Handle;
+							Pair.B = Body2Handle;
+							IgnorePairs.Add(Pair);
+						}
 					}
 				}
 
@@ -861,7 +870,6 @@ void FAnimNode_RigidBody::InitPhysics(const UAnimInstance* InAnimInstance)
 		HighLevelBodyInstances.Empty();
 		BodiesSorted.Empty();
 
-		TArray<ImmediatePhysics::FSimulation::FIgnorePair> IgnorePairs;
 		const TMap<FRigidBodyIndexPair, bool>& DisableTable = UsePhysicsAsset->CollisionDisableTable;
 		for(auto ConstItr = DisableTable.CreateConstIterator(); ConstItr; ++ConstItr)
 		{
