Chaos - RBAN: avoid some nastiness with world-space sime when scale goes to zero

diff --git a/Engine/Source/Runtime/AnimGraphRuntime/Private/BoneControllers/AnimNode_RigidBody.cpp b/Engine/Source/Runtime/AnimGraphRuntime/Private/BoneControllers/AnimNode_RigidBody.cpp
--- a/Engine/Source/Runtime/AnimGraphRuntime/Private/BoneControllers/AnimNode_RigidBody.cpp
+++ b/Engine/Source/Runtime/AnimGraphRuntime/Private/BoneControllers/AnimNode_RigidBody.cpp
@@ -13,6 +13,7 @@
 #include "Physics/ImmediatePhysics/ImmediatePhysicsStats.h"
 #include "PhysicsEngine/PhysicsSettings.h"
 #include "Logging/MessageLog.h"
+#include "Logging/LogMacros.h"
 
 //PRAGMA_DISABLE_OPTIMIZATION
 
@@ -25,6 +26,13 @@ DEFINE_STAT(STAT_RigidBodyNodeInitTime);
 
 CSV_DECLARE_CATEGORY_MODULE_EXTERN(ENGINE_API, Animation);
 
+#if UE_BUILD_SHIPPING || UE_BUILD_TEST
+DECLARE_LOG_CATEGORY_EXTERN(LogRBAN, Log, Warning);
+#else
+DECLARE_LOG_CATEGORY_EXTERN(LogRBAN, Log, All);
+#endif
+DEFINE_LOG_CATEGORY(LogRBAN);
+
 TAutoConsoleVariable<int32> CVarEnableRigidBodyNode(TEXT("p.RigidBodyNode"), 1, TEXT("Enables/disables rigid body node updates and evaluations"), ECVF_Default);
 TAutoConsoleVariable<int32> CVarRigidBodyLODThreshold(TEXT("p.RigidBodyLODThreshold"), -1, TEXT("Max LOD that rigid body node is allowed to run on. Provides a global threshold that overrides per-node the LODThreshold property. -1 means no override."), ECVF_Scalability);
 
@@ -60,6 +68,7 @@ FAnimNode_RigidBody::FAnimNode_RigidBody():
 	bTransferBoneVelocities = false;
 	bFreezeIncomingPoseOnStart = false;
 	bClampLinearTranslationLimitToRefPose = false;
+	WorldSpaceMinimumScale = 0.01f;
 
 	OverrideSolverIterations.SolverIterations = -1;
 	OverrideSolverIterations.JointIterations = -1;
@@ -268,6 +277,8 @@ void FAnimNode_RigidBody::EvaluateSkeletalControl_AnyThread(FComponentSpacePoseC
 		EvalCounter.SetMaxSkippedFrames(Output.AnimInstanceProxy->GetEvaluationCounter().GetMaxSkippedFrames());
 		if(!EvalCounter.WasSynchronizedLastFrame(Output.AnimInstanceProxy->GetEvaluationCounter())  && bRBAN_EnableTimeBasedReset)
 		{
+			UE_LOG(LogRBAN, Verbose, TEXT("%s Time-Based Reset"), *Output.AnimInstanceProxy->GetAnimInstanceName());
+
 			ResetSimulatedTeleportType = ETeleportType::ResetPhysics;
 		}
 	}
@@ -284,8 +295,14 @@ void FAnimNode_RigidBody::EvaluateSkeletalControl_AnyThread(FComponentSpacePoseC
 		{
 			PreviousCompWorldSpaceTM = CompWorldSpaceTM;
 		}
-		const FTransform BaseBoneTM = Output.Pose.GetComponentSpaceTransform(BaseBoneRef.GetCompactPoseIndex(BoneContainer));
 
+		// Disable simulation below minimum scale in world space mode. World space sim doesn't play nice with scale anyway - we do not scale joint offets or collision shapes.
+		if ((SimulationSpace == ESimulationSpace::WorldSpace) && (CompWorldSpaceTM.GetScale3D().SizeSquared() < WorldSpaceMinimumScale * WorldSpaceMinimumScale))
+		{
+			return;
+		}
+
+		const FTransform BaseBoneTM = Output.Pose.GetComponentSpaceTransform(BaseBoneRef.GetCompactPoseIndex(BoneContainer));
 		PhysicsSimulation->SetSimulationSpaceTransform(SpaceToWorldTransform(SimulationSpace, CompWorldSpaceTM, BaseBoneTM));
 
 		// Initialize potential new bodies because of LOD change.
@@ -349,6 +366,8 @@ void FAnimNode_RigidBody::EvaluateSkeletalControl_AnyThread(FComponentSpacePoseC
 			{
 				case ETeleportType::TeleportPhysics:
 				{
+					UE_LOG(LogRBAN, Verbose, TEXT("%s TeleportPhysics (Scale: %f %f %f)"), *Output.AnimInstanceProxy->GetAnimInstanceName(), CompWorldSpaceTM.GetScale3D().X, CompWorldSpaceTM.GetScale3D().Y, CompWorldSpaceTM.GetScale3D().Z);
+
 					// Teleport bodies.
 					for (const FOutputBoneData& OutputData : OutputBoneData)
 					{
@@ -378,6 +397,8 @@ void FAnimNode_RigidBody::EvaluateSkeletalControl_AnyThread(FComponentSpacePoseC
 
 				case ETeleportType::ResetPhysics:
 				{
+					UE_LOG(LogRBAN, Verbose, TEXT("%s ResetPhysics (Scale: %f %f %f)"), *Output.AnimInstanceProxy->GetAnimInstanceName(), CompWorldSpaceTM.GetScale3D().X, CompWorldSpaceTM.GetScale3D().Y, CompWorldSpaceTM.GetScale3D().Z);
+
 					// Completely reset bodies.
 					for (const FOutputBoneData& OutputData : OutputBoneData)
 					{
diff --git a/Engine/Source/Runtime/AnimGraphRuntime/Public/BoneControllers/AnimNode_RigidBody.h b/Engine/Source/Runtime/AnimGraphRuntime/Public/BoneControllers/AnimNode_RigidBody.h
--- a/Engine/Source/Runtime/AnimGraphRuntime/Public/BoneControllers/AnimNode_RigidBody.h
+++ b/Engine/Source/Runtime/AnimGraphRuntime/Public/BoneControllers/AnimNode_RigidBody.h
@@ -149,6 +149,12 @@ public:
 	UPROPERTY(EditAnywhere, Category = Settings)
 	uint8 bClampLinearTranslationLimitToRefPose : 1;
 
+	/**
+		For world-space simulations, if the magnitude of the component's 3D scale is less than WorldSpaceMinimumScale, do not update the node.
+	*/
+	UPROPERTY(EditAnywhere, Category = Settings)
+	float WorldSpaceMinimumScale;
+
 	/**
 	 * Solver iteration settings overrides (defaults are set in the Physics Asset).
 	 * These can be varied in the runtime and set through blueprint (e.g., to increase
diff --git a/Engine/Source/Runtime/Engine/Private/PhysicsEngine/ImmediatePhysics/ImmediatePhysicsChaos/ImmediatePhysicsActorHandle_Chaos.cpp b/Engine/Source/Runtime/Engine/Private/PhysicsEngine/ImmediatePhysics/ImmediatePhysicsChaos/ImmediatePhysicsActorHandle_Chaos.cpp
--- a/Engine/Source/Runtime/Engine/Private/PhysicsEngine/ImmediatePhysics/ImmediatePhysicsChaos/ImmediatePhysicsActorHandle_Chaos.cpp
+++ b/Engine/Source/Runtime/Engine/Private/PhysicsEngine/ImmediatePhysics/ImmediatePhysicsChaos/ImmediatePhysicsActorHandle_Chaos.cpp
@@ -390,13 +390,20 @@ namespace ImmediatePhysics_Chaos
 	{
 		using namespace Chaos;
 
-		FParticleUtilities::SetActorWorldTransform(TGenericParticleHandle<FReal, 3>(ParticleHandle), WorldTM);
+		FParticleUtilitiesXR::SetActorWorldTransform(TGenericParticleHandle<FReal, 3>(ParticleHandle), WorldTM);
+
+		if (auto* Kinematic = ParticleHandle->CastToKinematicParticle())
+		{
+			Kinematic->V() = FVec3(0);
+			Kinematic->W() = FVec3(0);
+			Kinematic->KinematicTarget().Clear();
+		}
 
 		auto* Dynamic = ParticleHandle->CastToRigidParticle();
 		if(Dynamic && Dynamic->ObjectState() == Chaos::EObjectStateType::Dynamic)
 		{
-			Dynamic->X() = Dynamic->P();
-			Dynamic->R() = Dynamic->Q();
+			Dynamic->P() = Dynamic->X();
+			Dynamic->Q() = Dynamic->R();
 			Dynamic->AuxilaryValue(ParticlePrevXs) = Dynamic->P();
 			Dynamic->AuxilaryValue(ParticlePrevRs) = Dynamic->Q();
 		}
diff --git a/Engine/Source/Runtime/Experimental/Chaos/Private/Chaos/Evolution/PBDMinEvolution.cpp b/Engine/Source/Runtime/Experimental/Chaos/Private/Chaos/Evolution/PBDMinEvolution.cpp
--- a/Engine/Source/Runtime/Experimental/Chaos/Private/Chaos/Evolution/PBDMinEvolution.cpp
+++ b/Engine/Source/Runtime/Experimental/Chaos/Private/Chaos/Evolution/PBDMinEvolution.cpp
@@ -284,7 +284,7 @@ namespace Chaos
 				// Nothing to do
 				break;
 
-			case EKinematicTargetMode::Zero:
+			case EKinematicTargetMode::Reset:
 			{
 				// Reset velocity and then switch to do-nothing mode
 				Particle.V() = FVec3(0);
@@ -303,7 +303,7 @@ namespace Chaos
 				{
 					TargetPos = KinematicTarget.GetTarget().GetLocation();
 					TargetRot = KinematicTarget.GetTarget().GetRotation();
-					KinematicTarget.SetMode(EKinematicTargetMode::Zero);
+					KinematicTarget.SetMode(EKinematicTargetMode::Reset);
 				}
 				else
 				{
diff --git a/Engine/Source/Runtime/Experimental/Chaos/Public/Chaos/KinematicGeometryParticles.h b/Engine/Source/Runtime/Experimental/Chaos/Public/Chaos/KinematicGeometryParticles.h
--- a/Engine/Source/Runtime/Experimental/Chaos/Public/Chaos/KinematicGeometryParticles.h
+++ b/Engine/Source/Runtime/Experimental/Chaos/Public/Chaos/KinematicGeometryParticles.h
@@ -12,9 +12,9 @@ namespace Chaos
 	enum class EKinematicTargetMode
 	{
 		None,			/** Particle does not move and no data is changed */
-		Zero,			/** Particle does not move, velocity and angular velocity are zeroed, then mode is set to "None". */
-		Position,		/** Particle is moved to Kinematic Target transform, velocity and angular velocity updated to reflect the change, then mode is reset to "Zero". */
-		Velocity,		/** Particle is moved based on velocity and angular velocity, mode remains as "Velocity". */
+		Reset,			/** Particle does not move, velocity and angular velocity are zeroed, then mode is set to "None". */
+		Position,		/** Particle is moved to Kinematic Target transform, velocity and angular velocity updated to reflect the change, then mode is set to "Zero". */
+		Velocity,		/** Particle is moved based on velocity and angular velocity, mode remains as "Velocity" until changed. */
 	};
 
 	/**
@@ -38,6 +38,9 @@ namespace Chaos
 		/** Get the target transform (asserts if not in Position mode) */
 		const TRigidTransform<T, d>& GetTarget() const { check(Mode == EKinematicTargetMode::Position); return Target; }
 
+		/** Clear the kinematic target */
+		void Clear() { Target = TRigidTransform<T, d>(); Mode = EKinematicTargetMode::None; }
+
 		/** Use transform target mode and set the transform target */
 		void SetTargetMode(const TRigidTransform<T, d>& InTarget) { Target = InTarget;  Mode = EKinematicTargetMode::Position; }
 
diff --git a/Engine/Source/Runtime/Experimental/Chaos/Public/Chaos/PBDRigidsEvolution.h b/Engine/Source/Runtime/Experimental/Chaos/Public/Chaos/PBDRigidsEvolution.h
--- a/Engine/Source/Runtime/Experimental/Chaos/Public/Chaos/PBDRigidsEvolution.h
+++ b/Engine/Source/Runtime/Experimental/Chaos/Public/Chaos/PBDRigidsEvolution.h
@@ -580,7 +580,7 @@ class TPBDRigidsEvolutionBase
 				// Nothing to do
 				break;
 
-			case EKinematicTargetMode::Zero:
+			case EKinematicTargetMode::Reset:
 			{
 				// Reset velocity and then switch to do-nothing mode
 				Particle.V() = FVec3(0, 0, 0);
@@ -599,7 +599,7 @@ class TPBDRigidsEvolutionBase
 				{
 					TargetPos = KinematicTarget.GetTarget().GetLocation();
 					TargetRot = KinematicTarget.GetTarget().GetRotation();
-					KinematicTarget.SetMode(EKinematicTargetMode::Zero);
+					KinematicTarget.SetMode(EKinematicTargetMode::Reset);
 				}
 				else
 				{
