Chaos - RBAN remove anim node solver iteration overrides

diff --git a/Engine/Source/Runtime/AnimGraphRuntime/Private/BoneControllers/AnimNode_RigidBody.cpp b/Engine/Source/Runtime/AnimGraphRuntime/Private/BoneControllers/AnimNode_RigidBody.cpp
--- a/Engine/Source/Runtime/AnimGraphRuntime/Private/BoneControllers/AnimNode_RigidBody.cpp
+++ b/Engine/Source/Runtime/AnimGraphRuntime/Private/BoneControllers/AnimNode_RigidBody.cpp
@@ -70,13 +70,6 @@ FAnimNode_RigidBody::FAnimNode_RigidBody():
 	bClampLinearTranslationLimitToRefPose = false;
 	WorldSpaceMinimumScale = 0.01f;
 
-	OverrideSolverIterations.SolverIterations = -1;
-	OverrideSolverIterations.JointIterations = -1;
-	OverrideSolverIterations.CollisionIterations = -1;
-	OverrideSolverIterations.SolverPushOutIterations = -1;
-	OverrideSolverIterations.JointPushOutIterations = -1;
-	OverrideSolverIterations.CollisionPushOutIterations = -1;
-
 	PreviousTransform = CurrentTransform = FTransform::Identity;
 	PreviousComponentLinearVelocity = FVector::ZeroVector;	
 
@@ -530,14 +523,6 @@ void FAnimNode_RigidBody::EvaluateSkeletalControl_AnyThread(FComponentSpacePoseC
 				SolverIterations.JointPushOutIterations,
 				SolverIterations.CollisionPushOutIterations
 			);
-			PhysicsSimulation->SetSolverIterations(
-				SolverIterations.FixedTimeStep,
-				OverrideSolverIterations.SolverIterations,
-				OverrideSolverIterations.JointIterations, 
-				OverrideSolverIterations.CollisionIterations, 
-				OverrideSolverIterations.SolverPushOutIterations, 
-				OverrideSolverIterations.JointPushOutIterations, 
-				OverrideSolverIterations.CollisionPushOutIterations);
 
 			PhysicsSimulation->Simulate_AssumesLocked(DeltaSeconds, MaxDeltaSeconds, MaxSteps, SimSpaceGravity);
 #endif
diff --git a/Engine/Source/Runtime/AnimGraphRuntime/Private/BoneControllers/AnimNode_RigidBody_Chaos.cpp b/Engine/Source/Runtime/AnimGraphRuntime/Private/BoneControllers/AnimNode_RigidBody_Chaos.cpp
--- a/Engine/Source/Runtime/AnimGraphRuntime/Private/BoneControllers/AnimNode_RigidBody_Chaos.cpp
+++ b/Engine/Source/Runtime/AnimGraphRuntime/Private/BoneControllers/AnimNode_RigidBody_Chaos.cpp
@@ -52,13 +52,6 @@ FAnimNode_RigidBody_Chaos::FAnimNode_RigidBody_Chaos()
 	bFreezeIncomingPoseOnStart = false;
 	bClampLinearTranslationLimitToRefPose = false;
 
-	OverrideSolverIterations.SolverIterations = -1;
-	OverrideSolverIterations.JointIterations = -1;
-	OverrideSolverIterations.CollisionIterations = -1;
-	OverrideSolverIterations.SolverPushOutIterations = -1;
-	OverrideSolverIterations.JointPushOutIterations = -1;
-	OverrideSolverIterations.CollisionPushOutIterations = -1;
-
 	PreviousTransform = CurrentTransform = FTransform::Identity;
 	PreviousComponentLinearVelocity = FVector::ZeroVector;	
 
@@ -478,14 +471,6 @@ void FAnimNode_RigidBody_Chaos::EvaluateSkeletalControl_AnyThread(FComponentSpac
 				SolverIterations.JointPushOutIterations,
 				SolverIterations.CollisionPushOutIterations
 			);
-			PhysicsSimulation->SetSolverIterations(
-				SolverIterations.FixedTimeStep,
-				OverrideSolverIterations.SolverIterations,
-				OverrideSolverIterations.JointIterations,
-				OverrideSolverIterations.CollisionIterations,
-				OverrideSolverIterations.SolverPushOutIterations,
-				OverrideSolverIterations.JointPushOutIterations,
-				OverrideSolverIterations.CollisionPushOutIterations);
 
 			// Run simulation at a minimum of 30 FPS to prevent system from exploding.
 			// DeltaTime can be higher due to URO, so take multiple iterations in that case.
diff --git a/Engine/Source/Runtime/AnimGraphRuntime/Public/BoneControllers/AnimNode_RigidBody.h b/Engine/Source/Runtime/AnimGraphRuntime/Public/BoneControllers/AnimNode_RigidBody.h
--- a/Engine/Source/Runtime/AnimGraphRuntime/Public/BoneControllers/AnimNode_RigidBody.h
+++ b/Engine/Source/Runtime/AnimGraphRuntime/Public/BoneControllers/AnimNode_RigidBody.h
@@ -155,15 +155,6 @@ public:
 	UPROPERTY(EditAnywhere, Category = Settings)
 	float WorldSpaceMinimumScale;
 
-	/**
-	 * Solver iteration settings overrides (defaults are set in the Physics Asset).
-	 * These can be varied in the runtime and set through blueprint (e.g., to increase
-	 * iterations during difficult movements).
-	 * Set to -1 to leave an individual iteration value at its Physics Asset value.
-	 */
-	UPROPERTY(EditAnywhere, Category = Settings, meta = (PinHiddenByDefault))
-	FSolverIterations OverrideSolverIterations;
-
 private:
 	uint8 bEnabled : 1;
 	uint8 bSimulationStarted : 1;
diff --git a/Engine/Source/Runtime/AnimGraphRuntime/Public/BoneControllers/AnimNode_RigidBody_Chaos.h b/Engine/Source/Runtime/AnimGraphRuntime/Public/BoneControllers/AnimNode_RigidBody_Chaos.h
--- a/Engine/Source/Runtime/AnimGraphRuntime/Public/BoneControllers/AnimNode_RigidBody_Chaos.h
+++ b/Engine/Source/Runtime/AnimGraphRuntime/Public/BoneControllers/AnimNode_RigidBody_Chaos.h
@@ -137,15 +137,6 @@ public:
 	UPROPERTY(EditAnywhere, Category = Settings)
 	uint8 bClampLinearTranslationLimitToRefPose : 1;
 
-	/**
-	 * Solver iteration settings overrides (defaults are set in the Physics Asset).
-	 * These can be varied in the runtime and set through blueprint (e.g., to increase
-	 * iterations during difficult movements).
-	 * Set to -1 to leave an individual iteration value at its Physics Asset value.
-	 */
-	UPROPERTY(EditAnywhere, Category = Settings, meta = (PinHiddenByDefault))
-		FSolverIterations OverrideSolverIterations;
-
 private:
 #if WITH_EDITORONLY_DATA
 	UPROPERTY()
