Chaos - RBAN fix ClampLinearTranslationLimit flag

diff --git a/Engine/Source/Runtime/AnimGraphRuntime/Private/BoneControllers/AnimNode_RigidBody.cpp b/Engine/Source/Runtime/AnimGraphRuntime/Private/BoneControllers/AnimNode_RigidBody.cpp
--- a/Engine/Source/Runtime/AnimGraphRuntime/Private/BoneControllers/AnimNode_RigidBody.cpp
+++ b/Engine/Source/Runtime/AnimGraphRuntime/Private/BoneControllers/AnimNode_RigidBody.cpp
@@ -235,7 +235,7 @@ void FAnimNode_RigidBody::InitializeNewBodyTransformsDuringSimulation(FComponent
 				}
 
 				const FTransform WSBodyTM = BodyRelativeTransform * Bodies[OutputData.ParentBodyIndex]->GetWorldTransform();
-				Bodies[BodyIndex]->SetWorldTransform(WSBodyTM);
+				Bodies[BodyIndex]->InitWorldTransform(WSBodyTM);
 				BodyAnimData[BodyIndex].RefPoseLength = BodyRelativeTransform.GetLocation().Size();
 			}
 			// If we don't have a parent body, then we can just grab the incoming pose in component space.
@@ -244,7 +244,7 @@ void FAnimNode_RigidBody::InitializeNewBodyTransformsDuringSimulation(FComponent
 				const FTransform& ComponentSpaceTM = Output.Pose.GetComponentSpaceTransform(OutputData.CompactPoseBoneIndex);
 				const FTransform BodyTM = ConvertCSTransformToSimSpace(SimulationSpace, ComponentSpaceTM, ComponentTransform, BaseBoneTM);
 
-				Bodies[BodyIndex]->SetWorldTransform(BodyTM);
+				Bodies[BodyIndex]->InitWorldTransform(BodyTM);
 			}
 		}
 	}
@@ -400,7 +400,7 @@ void FAnimNode_RigidBody::EvaluateSkeletalControl_AnyThread(FComponentSpacePoseC
 
 						const FTransform& ComponentSpaceTM = Output.Pose.GetComponentSpaceTransform(OutputData.CompactPoseBoneIndex);
 						const FTransform BodyTM = ConvertCSTransformToSimSpace(SimulationSpace, ComponentSpaceTM, CompWorldSpaceTM, BaseBoneTM);
-						Bodies[BodyIndex]->SetWorldTransform(BodyTM);
+						Bodies[BodyIndex]->InitWorldTransform(BodyTM);
 						if (OutputData.ParentBodyIndex != INDEX_NONE)
 						{
 							BodyAnimData[BodyIndex].RefPoseLength = BodyTM.GetRelativeTransform(Bodies[OutputData.ParentBodyIndex]->GetWorldTransform()).GetLocation().Size();
@@ -538,6 +538,7 @@ void FAnimNode_RigidBody::EvaluateSkeletalControl_AnyThread(FComponentSpacePoseC
 
 				// if we clamp translation, we only do this when all linear translation are locked
 				// 
+				// @todo(ccaulfield): this shouldn't be required with Chaos - projection should be handling it...
 				if (bClampLinearTranslationLimitToRefPose
 					&&BodyAnimData[BodyIndex].LinearXMotion == ELinearConstraintMotion::LCM_Locked
 					&& BodyAnimData[BodyIndex].LinearYMotion == ELinearConstraintMotion::LCM_Locked
diff --git a/Engine/Source/Runtime/Engine/Private/PhysicsEngine/ImmediatePhysics/ImmediatePhysicsChaos/ImmediatePhysicsActorHandle_Chaos.cpp b/Engine/Source/Runtime/Engine/Private/PhysicsEngine/ImmediatePhysics/ImmediatePhysicsChaos/ImmediatePhysicsActorHandle_Chaos.cpp
--- a/Engine/Source/Runtime/Engine/Private/PhysicsEngine/ImmediatePhysics/ImmediatePhysicsChaos/ImmediatePhysicsActorHandle_Chaos.cpp
+++ b/Engine/Source/Runtime/Engine/Private/PhysicsEngine/ImmediatePhysics/ImmediatePhysicsChaos/ImmediatePhysicsActorHandle_Chaos.cpp
@@ -386,11 +386,12 @@ namespace ImmediatePhysics_Chaos
 		}
 	}
 
-	void FActorHandle::SetWorldTransform(const FTransform& WorldTM)
+
+	void FActorHandle::InitWorldTransform(const FTransform& WorldTM)
 	{
 		using namespace Chaos;
 
-		FParticleUtilitiesXR::SetActorWorldTransform(TGenericParticleHandle<FReal, 3>(ParticleHandle), WorldTM);
+		SetWorldTransform(WorldTM);
 
 		if (auto* Kinematic = ParticleHandle->CastToKinematicParticle())
 		{
@@ -398,6 +399,13 @@ namespace ImmediatePhysics_Chaos
 			Kinematic->W() = FVec3(0);
 			Kinematic->KinematicTarget().Clear();
 		}
+	}
+
+	void FActorHandle::SetWorldTransform(const FTransform& WorldTM)
+	{
+		using namespace Chaos;
+
+		FParticleUtilitiesXR::SetActorWorldTransform(TGenericParticleHandle<FReal, 3>(ParticleHandle), WorldTM);
 
 		auto* Dynamic = ParticleHandle->CastToRigidParticle();
 		if(Dynamic && Dynamic->ObjectState() == Chaos::EObjectStateType::Dynamic)
diff --git a/Engine/Source/Runtime/Engine/Public/Physics/ImmediatePhysics/ImmediatePhysicsChaos/ImmediatePhysicsActorHandle_Chaos.h b/Engine/Source/Runtime/Engine/Public/Physics/ImmediatePhysics/ImmediatePhysicsChaos/ImmediatePhysicsActorHandle_Chaos.h
--- a/Engine/Source/Runtime/Engine/Public/Physics/ImmediatePhysics/ImmediatePhysicsChaos/ImmediatePhysicsActorHandle_Chaos.h
+++ b/Engine/Source/Runtime/Engine/Public/Physics/ImmediatePhysics/ImmediatePhysicsChaos/ImmediatePhysicsActorHandle_Chaos.h
@@ -20,7 +20,10 @@ namespace ImmediatePhysics_Chaos
 
 		void SetEnabled(bool bEnabled);
 
-		/** Sets the world transform.*/
+		/** Sets the world transform, zeroes velocity, etc.*/
+		void InitWorldTransform(const FTransform& WorldTM);
+
+		/** Sets the world transform, maintains velocity etc.*/
 		void SetWorldTransform(const FTransform& WorldTM);
 
 		/** Make a body kinematic, or non-kinematic */
diff --git a/Engine/Source/Runtime/Engine/Public/Physics/ImmediatePhysics/ImmediatePhysicsPhysX/ImmediatePhysicsActorHandle_PhysX.h b/Engine/Source/Runtime/Engine/Public/Physics/ImmediatePhysics/ImmediatePhysicsPhysX/ImmediatePhysicsActorHandle_PhysX.h
--- a/Engine/Source/Runtime/Engine/Public/Physics/ImmediatePhysics/ImmediatePhysicsPhysX/ImmediatePhysicsActorHandle_PhysX.h
+++ b/Engine/Source/Runtime/Engine/Public/Physics/ImmediatePhysics/ImmediatePhysicsPhysX/ImmediatePhysicsActorHandle_PhysX.h
@@ -21,6 +21,11 @@ namespace ImmediatePhysics_PhysX
 			OwningSimulation.GetLowLevelBody(ActorDataIndex).body2World = U2PTransform(ActorToBody * WorldTM);
 		}
 
+		void InitWorldTransform(const FTransform& WorldTM)
+		{
+			SetWorldTransform(WorldTM);
+		}
+
 		/** Make a body kinematic, or non-kinematic */
 		void SetIsKinematic(bool bKinematic)
 		{
