Sequencer: NLA fix for child bone rotations

diff --git a/Engine/Source/Editor/MovieSceneTools/Private/EditModes/SkeletalAnimationTrackEditMode.cpp b/Engine/Source/Editor/MovieSceneTools/Private/EditModes/SkeletalAnimationTrackEditMode.cpp
--- a/Engine/Source/Editor/MovieSceneTools/Private/EditModes/SkeletalAnimationTrackEditMode.cpp
+++ b/Engine/Source/Editor/MovieSceneTools/Private/EditModes/SkeletalAnimationTrackEditMode.cpp
@@ -194,7 +194,7 @@ static void DrawBones(const TArray<FBoneIndexType>& RequiredBones, const FRefere
 }
 
 static void DrawBonesFromCompactPose(const FCompactPose& Pose, USkeletalMeshComponent* MeshComponent, FPrimitiveDrawInterface* PDI,
-	const FLinearColor& DrawColour, FVector& LocationOfRootBone)
+	const FLinearColor& DrawColour, FCompactPoseBoneIndex RootBoneIndex, FVector& LocationOfRootBone)
 {
 	LocationOfRootBone = FVector(TNumericLimits<float>::Max(), TNumericLimits<float>::Max(), TNumericLimits<float>::Max());
 	if (Pose.GetNumBones() > 0 && MeshComponent != nullptr)
@@ -224,20 +224,19 @@ static void DrawBonesFromCompactPose(const FCompactPose& Pose, USkeletalMeshComp
 			if (ParentIndex == INDEX_NONE)
 			{
 				WorldTransforms[MeshBoneIndex.GetInt()] = Pose[BoneIndex] * MeshComponent->GetComponentTransform();
-				LocationOfRootBone.X = WorldTransforms[MeshBoneIndex.GetInt()].GetLocation().X;
-				LocationOfRootBone.Y = WorldTransforms[MeshBoneIndex.GetInt()].GetLocation().Y;
-				if (LocationOfRootBone.Z > WorldTransforms[MeshBoneIndex.GetInt()].GetLocation().Z)
-				{
-					LocationOfRootBone.Z = WorldTransforms[MeshBoneIndex.GetInt()].GetLocation().Z;
-				}
 			}
 			else
 			{
 				WorldTransforms[MeshBoneIndex.GetInt()] = Pose[BoneIndex] * WorldTransforms[ParentIndex];
-				if (LocationOfRootBone.Z > WorldTransforms[MeshBoneIndex.GetInt()].GetLocation().Z)
-				{
-					LocationOfRootBone.Z = WorldTransforms[MeshBoneIndex.GetInt()].GetLocation().Z;
-				}
+			}
+			if (RootBoneIndex == BoneIndex)
+			{
+				LocationOfRootBone.X = WorldTransforms[MeshBoneIndex.GetInt()].GetLocation().X;
+				LocationOfRootBone.Y = WorldTransforms[MeshBoneIndex.GetInt()].GetLocation().Y;
+			}
+			if (LocationOfRootBone.Z > WorldTransforms[MeshBoneIndex.GetInt()].GetLocation().Z)
+			{
+				LocationOfRootBone.Z = WorldTransforms[MeshBoneIndex.GetInt()].GetLocation().Z;
 			}
 			BoneColours[MeshBoneIndex.GetInt()] = DrawColour;
 		}
@@ -377,6 +376,7 @@ void FSkeletalAnimationTrackEditMode::Render(const FSceneView* View, FViewport*
 										UAnimSequence* AnimSequence = Cast<UAnimSequence>(AnimSection->Params.Animation);
 										if (AnimSequence)
 										{
+											int32 Index = AnimSection->SetBoneIndexForRootMotionCalculations(SkelAnimTrack->bBlendFirstChildOfRoot);
 											FCompactPose OutPose;
 											OutPose.ResetToRefPose(SkelMeshComp->GetAnimInstance()->GetRequiredBones());
 
@@ -402,10 +402,10 @@ void FSkeletalAnimationTrackEditMode::Render(const FSceneView* View, FViewport*
 											float Weight;
 											AnimSection->GetRootMotionTransform(TimeToSample, MovieScene->GetTickResolution(), RootMotionAtStart, Weight);
 											RootMotionAtStart = RootMotionAtStart * AnimSection->TempOffsetTransform;
-											OutPose[FCompactPoseBoneIndex(0)] = RootMotionAtStart;
+											OutPose[FCompactPoseBoneIndex(Index)] = RootMotionAtStart;
 
 											FVector RootLocation;
-											DrawBonesFromCompactPose(OutPose, SkelMeshComp, PDI, Colors[SectionIndex % 5], RootLocation);
+											DrawBonesFromCompactPose(OutPose, SkelMeshComp, PDI, Colors[SectionIndex % 5], FCompactPoseBoneIndex(Index), RootLocation);
 											if (IsRootSelected(AnimSection))
 											{
 												RootTransform = RootMotionAtStart;
diff --git a/Engine/Source/Editor/MovieSceneTools/Private/TrackEditors/SkeletalAnimationTrackEditor.cpp b/Engine/Source/Editor/MovieSceneTools/Private/TrackEditors/SkeletalAnimationTrackEditor.cpp
--- a/Engine/Source/Editor/MovieSceneTools/Private/TrackEditors/SkeletalAnimationTrackEditor.cpp
+++ b/Engine/Source/Editor/MovieSceneTools/Private/TrackEditors/SkeletalAnimationTrackEditor.cpp
@@ -1603,6 +1603,8 @@ void CopyInterpAnimControlTrack(TSharedRef<ISequencer> Sequencer, UInterpTrackAn
 
 void FSkeletalAnimationTrackEditor::BuildTrackContextMenu( FMenuBuilder& MenuBuilder, UMovieSceneTrack* Track )
 {
+	//there's a bug with a section being open already, so we end it.
+
 	TSharedPtr<ISequencer> SequencerPtr = GetSequencer();
 
 	UInterpTrackAnimControl* MatineeAnimControlTrack = nullptr;
@@ -1645,8 +1647,24 @@ void FSkeletalAnimationTrackEditor::BuildTrackContextMenu( FMenuBuilder& MenuBui
 	}
 	*/
 
-	MenuBuilder.BeginSection(NAME_None, LOCTEXT("SkelAnimTrackDisplay", "Display"));
+	MenuBuilder.BeginSection(NAME_None, LOCTEXT("SkelAnimRootMOtion", "Root Motion"));
 	{
+		MenuBuilder.AddMenuEntry(
+			NSLOCTEXT("Sequencer", "BlendFirstChildOfRoot", "Blend First Child Of Root"),
+			NSLOCTEXT("Sequencer", "BlendFirstChildOfRootTooltip", "If True, do not blend and match the root bones but instead the first child bone of the root. Toggle this on when the matched sequences in the track have no motion on the root."),
+			FSlateIcon(),
+			FUIAction(
+				FExecuteAction::CreateLambda([=]()->void {
+					SkeletalAnimationTrack->bBlendFirstChildOfRoot = SkeletalAnimationTrack->bBlendFirstChildOfRoot ? false : true;
+					SkeletalAnimationTrack->SetRootMotionsDirty();
+					SequencerPtr->NotifyMovieSceneDataChanged(EMovieSceneDataChangeType::TrackValueChanged);
+
+					}),
+				FCanExecuteAction::CreateLambda([=]()->bool { return SequencerPtr && SkeletalAnimationTrack != nullptr; }),
+						FIsActionChecked::CreateLambda([=]()->bool { return SkeletalAnimationTrack != nullptr && SkeletalAnimationTrack->bBlendFirstChildOfRoot; })),
+			NAME_None,
+			EUserInterfaceActionType::ToggleButton
+			);
 
 		MenuBuilder.AddMenuEntry(
 			NSLOCTEXT("Sequencer", "ShowRootMotionTrails", "Show Root Motion Trail"),
@@ -1662,12 +1680,11 @@ void FSkeletalAnimationTrackEditor::BuildTrackContextMenu( FMenuBuilder& MenuBui
 						FIsActionChecked::CreateLambda([=]()->bool { return SkeletalAnimationTrack != nullptr && SkeletalAnimationTrack->bShowRootMotionTrail; })),
 			NAME_None,
 						EUserInterfaceActionType::ToggleButton
-						);
+			);
 
-		MenuBuilder.EndSection();
 	}
 	MenuBuilder.EndSection();
-
+	MenuBuilder.AddSeparator();
 }
 
 TSharedPtr<SWidget> FSkeletalAnimationTrackEditor::BuildOutlinerEditWidget(const FGuid& ObjectBinding, UMovieSceneTrack* Track, const FBuildEditWidgetParams& Params)
diff --git a/Engine/Source/Runtime/AnimGraphRuntime/Private/AnimSequencerInstance.cpp b/Engine/Source/Runtime/AnimGraphRuntime/Private/AnimSequencerInstance.cpp
--- a/Engine/Source/Runtime/AnimGraphRuntime/Private/AnimSequencerInstance.cpp
+++ b/Engine/Source/Runtime/AnimGraphRuntime/Private/AnimSequencerInstance.cpp
@@ -35,7 +35,7 @@ void UAnimSequencerInstance::UpdateAnimTrack(UAnimSequenceBase* InAnimSequence,
 	GetProxyOnGameThread<FAnimSequencerInstanceProxy>().UpdateAnimTrack(InAnimSequence, SequenceId, InFromPosition, InToPosition, Weight, bFireNotifies);
 }
 
-void UAnimSequencerInstance::UpdateAnimTrackWithRootMotion(UAnimSequenceBase* InAnimSequence, int32 SequenceId, const TOptional<FTransform>& RootMotion, float InFromPosition, float InToPosition, float Weight, bool bFireNotifies)
+void UAnimSequencerInstance::UpdateAnimTrackWithRootMotion(UAnimSequenceBase* InAnimSequence, int32 SequenceId, const TOptional<FRootMotionOverride>& RootMotion, float InFromPosition, float InToPosition, float Weight, bool bFireNotifies)
 {
 	GetProxyOnGameThread<FAnimSequencerInstanceProxy>().UpdateAnimTrackWithRootMotion(InAnimSequence, SequenceId, RootMotion, InFromPosition, InToPosition, Weight, bFireNotifies);
 }
diff --git a/Engine/Source/Runtime/AnimGraphRuntime/Private/AnimSequencerInstanceProxy.cpp b/Engine/Source/Runtime/AnimGraphRuntime/Private/AnimSequencerInstanceProxy.cpp
--- a/Engine/Source/Runtime/AnimGraphRuntime/Private/AnimSequencerInstanceProxy.cpp
+++ b/Engine/Source/Runtime/AnimGraphRuntime/Private/AnimSequencerInstanceProxy.cpp
@@ -28,12 +28,35 @@ bool FAnimSequencerInstanceProxy::Evaluate(FPoseContext& Output)
 	SequencerRootNode.Evaluate_AnyThread(Output);
 	if (RootMotionOverride.IsSet())
 	{
-		for (const FCompactPoseBoneIndex BoneIndex : Output.Pose.ForEachBoneIndex())
+		// Almost always root is first, child is second bone but maybe not, so we do too loops to avoid any if's in the loop
+
+		if (!RootMotionOverride.GetValue().bBlendFirstChildOfRoot)
+		{
+			for (const FCompactPoseBoneIndex BoneIndex : Output.Pose.ForEachBoneIndex())
+			{
+				if (BoneIndex.IsRootBone())
+				{
+					Output.Pose[BoneIndex] = RootMotionOverride.GetValue().RootMotion;
+					break;
+				}
+			}
+		}
+		else
 		{
-			if (BoneIndex.IsRootBone())
+			int RootIndex = INDEX_NONE;
+			for (const FCompactPoseBoneIndex BoneIndex : Output.Pose.ForEachBoneIndex())
 			{
-				Output.Pose[BoneIndex] = RootMotionOverride.GetValue();
-				break;
+				FMeshPoseBoneIndex MeshBoneIndex = Output.Pose.GetBoneContainer().MakeMeshPoseIndex(BoneIndex);
+				int32 ParentIndex = Output.Pose.GetBoneContainer().GetParentBoneIndex(MeshBoneIndex.GetInt());
+				if (ParentIndex == INDEX_NONE)
+				{
+					RootIndex = MeshBoneIndex.GetInt();
+				}
+				else if (ParentIndex == RootIndex)
+				{
+					Output.Pose[BoneIndex] = RootMotionOverride.GetValue().RootMotion;
+					break;
+				}
 			}
 		}
 	}
@@ -151,21 +174,21 @@ void FAnimSequencerInstanceProxy::TermAnimTrack(int32 SequenceId)
 
 void FAnimSequencerInstanceProxy::UpdateAnimTrack(UAnimSequenceBase* InAnimSequence, uint32 SequenceId, float InPosition, float Weight, bool bFireNotifies)
 {
-	UpdateAnimTrack(InAnimSequence, SequenceId, TOptional<FTransform>(), TOptional<float>(), InPosition, Weight, bFireNotifies);
+	UpdateAnimTrack(InAnimSequence, SequenceId, TOptional<FRootMotionOverride>(), TOptional<float>(), InPosition, Weight, bFireNotifies);
 }
 
 void FAnimSequencerInstanceProxy::UpdateAnimTrack(UAnimSequenceBase* InAnimSequence, uint32 SequenceId, TOptional<float> InFromPosition, float InToPosition, float Weight, bool bFireNotifies)
 {
-	UpdateAnimTrack(InAnimSequence, SequenceId, TOptional<FTransform>(), InFromPosition, InToPosition, Weight, bFireNotifies);
+	UpdateAnimTrack(InAnimSequence, SequenceId, TOptional<FRootMotionOverride>(), InFromPosition, InToPosition, Weight, bFireNotifies);
 }
 
 
-void FAnimSequencerInstanceProxy::UpdateAnimTrackWithRootMotion(UAnimSequenceBase* InAnimSequence, int32 SequenceId, const TOptional<FTransform>& RootMotion, float InFromPosition, float InToPosition, float Weight, bool bFireNotifies)
+void FAnimSequencerInstanceProxy::UpdateAnimTrackWithRootMotion(UAnimSequenceBase* InAnimSequence, int32 SequenceId, const TOptional<FRootMotionOverride>& RootMotion, float InFromPosition, float InToPosition, float Weight, bool bFireNotifies)
 {
 	UpdateAnimTrack(InAnimSequence, SequenceId, RootMotion, InFromPosition, InToPosition, Weight, bFireNotifies);
 
 }
-void FAnimSequencerInstanceProxy::UpdateAnimTrack(UAnimSequenceBase* InAnimSequence, uint32 SequenceId, const TOptional<FTransform>& InRootMotionOverride, TOptional<float> InFromPosition, float InToPosition, float Weight, bool bFireNotifies)
+void FAnimSequencerInstanceProxy::UpdateAnimTrack(UAnimSequenceBase* InAnimSequence, uint32 SequenceId, const TOptional<FRootMotionOverride>& InRootMotionOverride, TOptional<float> InFromPosition, float InToPosition, float Weight, bool bFireNotifies)
 {
 	EnsureAnimTrack(InAnimSequence, SequenceId);
 
diff --git a/Engine/Source/Runtime/AnimGraphRuntime/Public/AnimSequencerInstance.h b/Engine/Source/Runtime/AnimGraphRuntime/Public/AnimSequencerInstance.h
--- a/Engine/Source/Runtime/AnimGraphRuntime/Public/AnimSequencerInstance.h
+++ b/Engine/Source/Runtime/AnimGraphRuntime/Public/AnimSequencerInstance.h
@@ -11,6 +11,8 @@
 #include "SequencerAnimationSupport.h"
 #include "AnimSequencerInstance.generated.h"
 
+struct FRootMotionOverride;
+
 UCLASS(transient, NotBlueprintable)
 class ANIMGRAPHRUNTIME_API UAnimSequencerInstance : public UAnimInstance, public ISequencerAnimationSupport
 {
@@ -23,7 +25,7 @@ public:
 	virtual void UpdateAnimTrack(UAnimSequenceBase* InAnimSequence, int32 SequenceId, float InFromPosition, float InToPosition, float Weight, bool bFireNotifies);
 
 	/** Update with Root Motion*/
-	void UpdateAnimTrackWithRootMotion(UAnimSequenceBase* InAnimSequence, int32 SequenceId,const TOptional<FTransform>& RootMotion, float InFromPosition, float InToPosition, float Weight, bool bFireNotifies);
+	void UpdateAnimTrackWithRootMotion(UAnimSequenceBase* InAnimSequence, int32 SequenceId,const TOptional<FRootMotionOverride>& RootMotion, float InFromPosition, float InToPosition, float Weight, bool bFireNotifies);
 
 	/** Construct all nodes in this instance */
 	virtual void ConstructNodes() override;
diff --git a/Engine/Source/Runtime/AnimGraphRuntime/Public/AnimSequencerInstanceProxy.h b/Engine/Source/Runtime/AnimGraphRuntime/Public/AnimSequencerInstanceProxy.h
--- a/Engine/Source/Runtime/AnimGraphRuntime/Public/AnimSequencerInstanceProxy.h
+++ b/Engine/Source/Runtime/AnimGraphRuntime/Public/AnimSequencerInstanceProxy.h
@@ -45,6 +45,16 @@ protected:
 
 };
 
+
+/** Optional Override To Specify RootMotion*/
+struct FRootMotionOverride
+{
+	FRootMotionOverride() :bBlendFirstChildOfRoot(false) {};
+	/** If true we use the first child of the root, if not we just use the root*/
+	bool bBlendFirstChildOfRoot;
+	FTransform RootMotion;
+};
+
 /** Quick n dirty RTTI to allow for derived classes to insert nodes of different types */
 #define SEQUENCER_INSTANCE_PLAYER_TYPE(TYPE, BASE) \
 	static const FName& GetTypeId() { static FName Type(TEXT(#TYPE)); return Type; } \
@@ -54,10 +64,11 @@ protected:
 struct FSequencerPlayerAnimSequence : public FSequencerPlayerBase
 {
 	SEQUENCER_INSTANCE_PLAYER_TYPE(FSequencerPlayerAnimSequence, FSequencerPlayerBase)
-	TOptional<FTransform> RootMotion;
+	TOptional<FRootMotionOverride> RootMotion;
 	struct FAnimNode_SequenceEvaluator PlayerNode;
 };
 
+
 /** Proxy override for this UAnimInstance-derived class */
 USTRUCT()
 struct ANIMGRAPHRUNTIME_API FAnimSequencerInstanceProxy : public FAnimInstanceProxy
@@ -85,7 +96,7 @@ public:
 	void UpdateAnimTrack(UAnimSequenceBase* InAnimSequence, uint32 SequenceId, float InPosition, float Weight, bool bFireNotifies);
 	void UpdateAnimTrack(UAnimSequenceBase* InAnimSequence, uint32 SequenceId, TOptional<float> InFromPosition, float InToPosition, float Weight, bool bFireNotifies);
 
-	void UpdateAnimTrackWithRootMotion(UAnimSequenceBase* InAnimSequence, int32 SequenceId, const TOptional<FTransform>& RootMotion, float InFromPosition, float InToPosition, float Weight, bool bFireNotifies);
+	void UpdateAnimTrackWithRootMotion(UAnimSequenceBase* InAnimSequence, int32 SequenceId, const TOptional<FRootMotionOverride>& RootMotion, float InFromPosition, float InToPosition, float Weight, bool bFireNotifies);
 
 	/** Reset all nodes in this instance */
 	virtual void ResetNodes();
@@ -98,7 +109,7 @@ public:
 
 protected:
 
-	void UpdateAnimTrack(UAnimSequenceBase* InAnimSequence, uint32 SequenceId, const TOptional<FTransform>& FTransform, TOptional<float> InFromPosition, float InToPosition, float Weight, bool bFireNotifies);
+	void UpdateAnimTrack(UAnimSequenceBase* InAnimSequence, uint32 SequenceId, const TOptional<FRootMotionOverride>& RootMomtionOverride, TOptional<float> InFromPosition, float InToPosition, float Weight, bool bFireNotifies);
 
 	/** Find a player of a specified type */
 	template<typename Type>
@@ -123,7 +134,7 @@ protected:
 	TMap<uint32, FSequencerPlayerBase*> SequencerToPlayerMap;
 
 	/** custom root motion override sent in from sequencer */
-	TOptional<FTransform> RootMotionOverride;
+	TOptional<FRootMotionOverride> RootMotionOverride;
 
 	void InitAnimTrack(UAnimSequenceBase* InAnimSequence, uint32 SequenceId);
 	void EnsureAnimTrack(UAnimSequenceBase* InAnimSequence, uint32 SequenceId);
diff --git a/Engine/Source/Runtime/MovieSceneTracks/Private/Evaluation/MovieSceneSkeletalAnimationTemplate.cpp b/Engine/Source/Runtime/MovieSceneTracks/Private/Evaluation/MovieSceneSkeletalAnimationTemplate.cpp
--- a/Engine/Source/Runtime/MovieSceneTracks/Private/Evaluation/MovieSceneSkeletalAnimationTemplate.cpp
+++ b/Engine/Source/Runtime/MovieSceneTracks/Private/Evaluation/MovieSceneSkeletalAnimationTemplate.cpp
@@ -15,6 +15,7 @@
 #include "Evaluation/IMovieSceneMotionVectorSimulation.h"
 #include "UObject/StrongObjectPtr.h"
 #include "SkeletalMeshRestoreState.h"
+#include "AnimSequencerInstanceProxy.h"
 
 bool ShouldUsePreviewPlayback(IMovieScenePlayer& Player, UObject& RuntimeObject)
 {
@@ -450,11 +451,16 @@ namespace MovieScene
 				FMovieSceneAnimTypeID AnimTypeID = SectionToAnimationIDs.GetAnimTypeID(Section);
 
 				Player.SavePreAnimatedState(*SequencerInst, AnimTypeID, FStatelessPreAnimatedTokenProducer(&ResetAnimSequencerInstance));
+				TOptional <FRootMotionOverride> RootMotion;
+				bool bBlendFirstChildOfRoot = false;
+				TOptional<FTransform> RootMotionTransform = AnimSection->GetRootMotion(CurrentTime, bBlendFirstChildOfRoot);
+				if (RootMotionTransform.IsSet())
+				{
+					RootMotion = FRootMotionOverride();
+					RootMotion.GetValue().RootMotion = RootMotionTransform.GetValue();
+					RootMotion.GetValue().bBlendFirstChildOfRoot = bBlendFirstChildOfRoot;
+				}
 
-				// Set position and weight
-				//SequencerInst->UpdateAnimTrack(InAnimSequence, GetTypeHash(AnimTypeID), InFromPosition, InToPosition, Weight, bFireNotifies);
-
-				TOptional<FTransform> RootMotion = AnimSection->GetRootMotion(CurrentTime);
 				SequencerInst->UpdateAnimTrackWithRootMotion(InAnimSequence, GetTypeHash(AnimTypeID),RootMotion, InFromPosition, InToPosition, Weight, bFireNotifies);
 
 			}
@@ -506,8 +512,15 @@ namespace MovieScene
 				Player.SavePreAnimatedState(*SequencerInst, AnimTypeID, FStatelessPreAnimatedTokenProducer(&ResetAnimSequencerInstance));
 
 				// Set position and weight
-				//SequencerInst->UpdateAnimTrack(InAnimSequence, GetTypeHash(AnimTypeID), InFromPosition, InToPosition, Weight, bFireNotifies);
-				TOptional<FTransform> RootMotion = AnimSection->GetRootMotion(CurrentTime);
+				TOptional <FRootMotionOverride> RootMotion;
+				bool bBlendFirstChildOfRoot = false;
+				TOptional<FTransform> RootMotionTransform = AnimSection->GetRootMotion(CurrentTime, bBlendFirstChildOfRoot);
+				if (RootMotionTransform.IsSet())
+				{
+					RootMotion = FRootMotionOverride();
+					RootMotion.GetValue().RootMotion = RootMotionTransform.GetValue();
+					RootMotion.GetValue().bBlendFirstChildOfRoot = bBlendFirstChildOfRoot;
+				}
 				SequencerInst->UpdateAnimTrackWithRootMotion(InAnimSequence, GetTypeHash(AnimTypeID), RootMotion, InFromPosition, InToPosition, Weight, bFireNotifies);
 
 			}
diff --git a/Engine/Source/Runtime/MovieSceneTracks/Private/Sections/MovieSceneSkeletalAnimationSection.cpp b/Engine/Source/Runtime/MovieSceneTracks/Private/Sections/MovieSceneSkeletalAnimationSection.cpp
--- a/Engine/Source/Runtime/MovieSceneTracks/Private/Sections/MovieSceneSkeletalAnimationSection.cpp
+++ b/Engine/Source/Runtime/MovieSceneTracks/Private/Sections/MovieSceneSkeletalAnimationSection.cpp
@@ -64,6 +64,7 @@ UMovieSceneSkeletalAnimationSection::UMovieSceneSkeletalAnimationSection( const
 	bMatchRotationPitch = false;
 	bMatchIncludeZHeight = false;
 
+	TempRootBoneIndex = 0;
 #if WITH_EDITOR
 
 	PreviousPlayRate = Params.PlayRate;
@@ -377,18 +378,20 @@ void UMovieSceneSkeletalAnimationSection::PostEditUndo()
 
 #endif
 
-TOptional<FTransform> UMovieSceneSkeletalAnimationSection::GetRootMotion(FFrameTime CurrentTime) const
+TOptional<FTransform> UMovieSceneSkeletalAnimationSection::GetRootMotion(FFrameTime CurrentTime, bool& bBlendFirstChildOfRoot) const
 {
 	if (GetRootMotionParams())
 	{
+		UMovieSceneSkeletalAnimationTrack* Track = GetTypedOuter<UMovieSceneSkeletalAnimationTrack>();
 		if (GetRootMotionParams()->bRootMotionsDirty)
 		{
-			UMovieSceneSkeletalAnimationTrack* Track = GetTypedOuter<UMovieSceneSkeletalAnimationTrack>();
 			if (Track)
 			{
 				Track->SetUpRootMotions(true);
 			}
 		}
+
+		bBlendFirstChildOfRoot = Track ? Track->bBlendFirstChildOfRoot : false;
 		return GetRootMotionParams()->GetRootMotion(CurrentTime);
 	}
 	return TOptional<FTransform>();
@@ -433,7 +436,41 @@ FTransform  UMovieSceneSkeletalAnimationSection::GetOffsetTransform() const
 	return OffsetTransform;
 }
 
-bool UMovieSceneSkeletalAnimationSection::GetRootMotionTransform(FFrameTime CurrentTime, FFrameRate FrameRate, FTransform& OutTransform, float& OutWeight) const
+int32 UMovieSceneSkeletalAnimationSection::SetBoneIndexForRootMotionCalculations(bool bBlendFirstChildOfRoot)
+{
+	TempRootBoneIndex = 0;
+	if (!bBlendFirstChildOfRoot)
+	{
+		return TempRootBoneIndex.GetValue();
+	}
+	else if(UAnimSequence* AnimSequence = Cast<UAnimSequence>(Params.Animation))
+	{ 
+		//but if not first find first
+		int32 RootIndex = INDEX_NONE;
+		for (int32 TrackIndex = 0; TrackIndex < AnimSequence->GetRawAnimationData().Num(); ++TrackIndex)
+		{
+			// verify if this bone exists in skeleton
+			int32 BoneTreeIndex = AnimSequence->GetSkeletonIndexFromRawDataTrackIndex(TrackIndex);
+			if (BoneTreeIndex != INDEX_NONE)
+			{
+				int32 ParentIndex = AnimSequence->GetSkeleton()->GetReferenceSkeleton().GetParentIndex(BoneTreeIndex);
+				if (ParentIndex == INDEX_NONE)
+				{
+					RootIndex = TrackIndex;
+				}
+				else if (ParentIndex == RootIndex)
+				{
+					TempRootBoneIndex = TrackIndex;
+					break;
+				}
+			}
+		}
+	}
+	return TempRootBoneIndex.GetValue();
+}
+
+
+bool UMovieSceneSkeletalAnimationSection::GetRootMotionTransform(FFrameTime CurrentTime, FFrameRate FrameRate,FTransform& OutTransform, float& OutWeight) const
 {
 	UAnimSequence* AnimSequence = Cast<UAnimSequence>(Params.Animation);
 	FTransform OffsetTransform = GetOffsetTransform();
@@ -448,7 +485,12 @@ bool UMovieSceneSkeletalAnimationSection::GetRootMotionTransform(FFrameTime Curr
 		{
 			CurrentTimeSeconds = 0.0f;
 		}
+	
 		OutTransform = AnimSequence->ExtractRootTrackTransform(CurrentTimeSeconds,nullptr);
+		if (TempRootBoneIndex.IsSet())
+		{
+			AnimSequence->GetBoneTransform(OutTransform, TempRootBoneIndex.GetValue(), CurrentTimeSeconds, true);
+		}
 		OutTransform = OutTransform * OffsetTransform;
 
 		return true;
diff --git a/Engine/Source/Runtime/MovieSceneTracks/Private/Tracks/MovieSceneSkeletalAnimationTrack.cpp b/Engine/Source/Runtime/MovieSceneTracks/Private/Tracks/MovieSceneSkeletalAnimationTrack.cpp
--- a/Engine/Source/Runtime/MovieSceneTracks/Private/Tracks/MovieSceneSkeletalAnimationTrack.cpp
+++ b/Engine/Source/Runtime/MovieSceneTracks/Private/Tracks/MovieSceneSkeletalAnimationTrack.cpp
@@ -26,6 +26,7 @@
 UMovieSceneSkeletalAnimationTrack::UMovieSceneSkeletalAnimationTrack(const FObjectInitializer& ObjectInitializer)
 	: Super(ObjectInitializer)
 	, bUseLegacySectionIndexBlend(false)
+	, bBlendFirstChildOfRoot(false)
 {
 #if WITH_EDITORONLY_DATA
 	TrackTint = FColor(124, 15, 124, 65);
@@ -94,14 +95,21 @@ void UMovieSceneSkeletalAnimationTrack::PostLoad()
 #if WITH_EDITOR
 void UMovieSceneSkeletalAnimationTrack::PostEditImport()
 {
+	Super::PostEditImport();
 	SetUpRootMotions(true);
 }
 
 void UMovieSceneSkeletalAnimationTrack::PostEditUndo()
 {
+	Super::PostEditUndo();
 	SetUpRootMotions(true);
 }
 
+void UMovieSceneSkeletalAnimationTrack::PostEditChangeProperty(FPropertyChangedEvent& PropertyChangedEvent)
+{
+
+	SetUpRootMotions(true);
+}
 
 #endif
 
@@ -452,6 +460,7 @@ void UMovieSceneSkeletalAnimationTrack::SetUpRootMotions(bool bForce)
 					AnimSection->TempOffsetTransform = FTransform::Identity;
 				}
 				PrevAnimSection = AnimSection;
+				AnimSection->SetBoneIndexForRootMotionCalculations(bBlendFirstChildOfRoot);
 			}
 		}
 
@@ -469,12 +478,8 @@ void UMovieSceneSkeletalAnimationTrack::SetUpRootMotions(bool bForce)
 		TArray<FTransform> CurrentTransforms;
 		TArray<float> CurrentWeights;
 
-
-		//UMovieSceneSkeletalAnimationSection* AnimSection = Cast< UMovieSceneSkeletalAnimationSection>(AnimationSections[SmallestActiveSection]);
-		//AnimSection->AllocateRootMotionCache(FrameRateInSeconds);
 		FFrameTime PreviousFrame = RootMotionParams.StartFrame;
 		int32 Index = 0;
-		
 		for (FFrameTime FrameNumber = RootMotionParams.StartFrame; FrameNumber <= RootMotionParams.EndFrame; FrameNumber += RootMotionParams.FrameTick)
 		{
 			CurrentTransforms.SetNum(0);
@@ -484,11 +489,11 @@ void UMovieSceneSkeletalAnimationTrack::SetUpRootMotions(bool bForce)
 			UMovieSceneSkeletalAnimationSection* PrevSection = nullptr;
 			for (UMovieSceneSection* Section : AnimationSections)
 			{
-				if (Section->GetRange().Contains(FrameNumber.FrameNumber))
+				if (Section && Section->GetRange().Contains(FrameNumber.FrameNumber))
 				{
-					//mz todo handles 
-					UMovieSceneSkeletalAnimationSection* AnimSection = Cast<UMovieSceneSkeletalAnimationSection>(Section);
-					if (AnimSection->GetRootMotionTransform( FrameNumber.FrameNumber, TickResolution, CurrentTransform, CurrentWeight))
+					UMovieSceneSkeletalAnimationSection* AnimSection = CastChecked<UMovieSceneSkeletalAnimationSection>(Section);
+				
+					if (AnimSection->GetRootMotionTransform(FrameNumber.FrameNumber, TickResolution, CurrentTransform, CurrentWeight))
 					{
 						CurrentTransform =  CurrentTransform * AnimSection->TempOffsetTransform;
 						CurrentTransforms.Add(CurrentTransform);
diff --git a/Engine/Source/Runtime/MovieSceneTracks/Public/Sections/MovieSceneSkeletalAnimationSection.h b/Engine/Source/Runtime/MovieSceneTracks/Public/Sections/MovieSceneSkeletalAnimationSection.h
--- a/Engine/Source/Runtime/MovieSceneTracks/Public/Sections/MovieSceneSkeletalAnimationSection.h
+++ b/Engine/Source/Runtime/MovieSceneTracks/Public/Sections/MovieSceneSkeletalAnimationSection.h
@@ -191,14 +191,16 @@ public:
 	//Temporary transform used to specify the global OffsetTransform while calculting the root motions.
 	FTransform TempOffsetTransform;
 
-// Functions/params related to root motion calcuations.
+	//Temporary index used by GetRootMotionTransform and set by SetBoneIndexforRootMotionCalculations
+	TOptional<int32> TempRootBoneIndex;
 public:
-	
+	// Functions/params related to root motion calcuations.
 	MOVIESCENETRACKS_API FMovieSceneSkeletalAnimRootMotionTrackParams* GetRootMotionParams() const;
 
 	MOVIESCENETRACKS_API FTransform GetOffsetTransform() const;
 	MOVIESCENETRACKS_API bool GetRootMotionVelocity(FFrameTime PreviousTime, FFrameTime CurrentTime, FFrameRate FrameRate, FTransform& OutVelocity, float& OutWeight) const;
-	MOVIESCENETRACKS_API bool GetRootMotionTransform(FFrameTime CurrentTime, FFrameRate FrameRate, FTransform& OutTransform, float& OutWeight) const;
+	MOVIESCENETRACKS_API int32 SetBoneIndexForRootMotionCalculations(bool bBlendFirstChildOfRoot);
+	MOVIESCENETRACKS_API bool GetRootMotionTransform(FFrameTime CurrentTime, FFrameRate FrameRate, FTransform& OutTransform,  float& OutWeight) const;
 	MOVIESCENETRACKS_API void MatchSectionByBoneTransform(USkeletalMeshComponent* SkelMeshComp, FFrameTime CurrentFrame, FFrameRate FrameRate,
 		const FName& BoneName); //add options for z and rotation
 
@@ -206,7 +208,7 @@ public:
 	
 	MOVIESCENETRACKS_API void FindBestBlendPoint(USkeletalMeshComponent* SkelMeshComp);
 
-	MOVIESCENETRACKS_API TOptional<FTransform> GetRootMotion(FFrameTime CurrentTime) const;
+	MOVIESCENETRACKS_API TOptional<FTransform> GetRootMotion(FFrameTime CurrentTime, bool& bBlendFirstChildOfRoot) const;
 
 	MOVIESCENETRACKS_API void ToggleMatchTranslation();
 
diff --git a/Engine/Source/Runtime/MovieSceneTracks/Public/Tracks/MovieSceneSkeletalAnimationTrack.h b/Engine/Source/Runtime/MovieSceneTracks/Public/Tracks/MovieSceneSkeletalAnimationTrack.h
--- a/Engine/Source/Runtime/MovieSceneTracks/Public/Tracks/MovieSceneSkeletalAnimationTrack.h
+++ b/Engine/Source/Runtime/MovieSceneTracks/Public/Tracks/MovieSceneSkeletalAnimationTrack.h
@@ -60,6 +60,8 @@ public:
 #if WITH_EDITOR
 	virtual void PostEditImport() override;
 	virtual void PostEditUndo() override;
+	virtual void PostEditChangeProperty(FPropertyChangedEvent& PropertyChangedEvent) override;
+
 #endif
 
 	virtual void RemoveAllAnimationData() override;
@@ -91,10 +93,10 @@ private:
 
 public:
 
-	void SetRootMotionsDirty();
-	void SetUpRootMotions(bool bForce);
-	void FindBestBlendPoint(USkeletalMeshComponent* SkelMeshComp,UMovieSceneSkeletalAnimationSection* FirstSection);
-	void MatchSectionByBoneTransform(bool bMatchWithPrevious, USkeletalMeshComponent* SkelMeshComp, UMovieSceneSkeletalAnimationSection* CurrentSection, FFrameTime CurrentFrame, FFrameRate FrameRate,
+	MOVIESCENETRACKS_API void SetRootMotionsDirty();
+	MOVIESCENETRACKS_API void SetUpRootMotions(bool bForce);
+	MOVIESCENETRACKS_API void FindBestBlendPoint(USkeletalMeshComponent* SkelMeshComp,UMovieSceneSkeletalAnimationSection* FirstSection);
+	MOVIESCENETRACKS_API void MatchSectionByBoneTransform(bool bMatchWithPrevious, USkeletalMeshComponent* SkelMeshComp, UMovieSceneSkeletalAnimationSection* CurrentSection, FFrameTime CurrentFrame, FFrameRate FrameRate,
 		const FName& BoneName, FTransform& SecondSectionRootDiff, FVector& TranslationDiff, FQuat& RotationDiff); //add options for z and for rotation.
 	MOVIESCENETRACKS_API void ToggleAutoMatchClipsRootMotions();
 #if WITH_EDITORONLY_DATA
@@ -122,6 +124,9 @@ public:
 	UPROPERTY()
 	FMovieSceneSkeletalAnimRootMotionTrackParams RootMotionParams;
 
+	/** Whether to blend and adjust the first child node instead of the root, this should be true for blending when the root is static, false if the animations have proper root motion*/
+	UPROPERTY(EditAnywhere, Category = "Root Motions")
+	bool bBlendFirstChildOfRoot;
 
 #if WITH_EDITORONLY_DATA
 public:
