Fix for LegIK node popping when angular constraints are on in tight poses

diff --git a/Engine/Source/Runtime/AnimGraphRuntime/Private/BoneControllers/AnimNode_LegIK.cpp b/Engine/Source/Runtime/AnimGraphRuntime/Private/BoneControllers/AnimNode_LegIK.cpp
--- a/Engine/Source/Runtime/AnimGraphRuntime/Private/BoneControllers/AnimNode_LegIK.cpp
+++ b/Engine/Source/Runtime/AnimGraphRuntime/Private/BoneControllers/AnimNode_LegIK.cpp
@@ -369,7 +369,10 @@ bool FAnimNode_LegIK::DoLegReachIK(FAnimLegIKData& InLegData)
 	const FVector FootIKLocation = InLegData.IKFootTransform.GetLocation();
 
 	// If we're already reaching our IK Target, we have no work to do.
-	if (FootFKLocation.Equals(FootIKLocation, ReachPrecision))
+	// Unless, we are applying a limit. In that case, must run solver to apply angular limit
+	const bool FootAtGoal = FootFKLocation.Equals(FootIKLocation, ReachPrecision);
+	const bool UsingAngleLimit = InLegData.LegDefPtr->bEnableRotationLimit;
+	if (FootAtGoal && !UsingAngleLimit)
 	{
 		return false;
 	}
@@ -624,7 +627,7 @@ void FIKChain::SolveFABRIK(const FVector& InTargetLocation, float InReachPrecisi
 
 	// Check distance between foot and foot target location
 	float Slop = FVector::Dist(Links[0].Location, InTargetLocation);
-	if (Slop > ReachPrecision)
+	if (Slop > ReachPrecision || bEnableRotationLimit)
 	{
 		if (bEnableRotationLimit)
 		{
