Fix for warning on FAnimNode_LayeredBoneBlend
#rnx

diff --git a/Engine/Source/Runtime/AnimGraphRuntime/Public/AnimNodes/AnimNode_LayeredBoneBlend.h b/Engine/Source/Runtime/AnimGraphRuntime/Public/AnimNodes/AnimNode_LayeredBoneBlend.h
--- a/Engine/Source/Runtime/AnimGraphRuntime/Public/AnimNodes/AnimNode_LayeredBoneBlend.h
+++ b/Engine/Source/Runtime/AnimGraphRuntime/Public/AnimNodes/AnimNode_LayeredBoneBlend.h
@@ -104,7 +104,8 @@ protected:
 
 public:	
 	FAnimNode_LayeredBoneBlend()
-		: bMeshSpaceRotationBlend(false)
+		: BlendMode(ELayeredBoneBlendMode::BranchFilter)
+		, bMeshSpaceRotationBlend(false)
 		, bMeshSpaceScaleBlend(false)
 		, CurveBlendOption(ECurveBlendOption::Override)
 		, bBlendRootMotionBasedOnRootBone(true)
