Fixed root motion weight modifier being incorrectly initialized

diff --git a/Engine/Source/Runtime/AnimGraphRuntime/Private/AnimNodes/AnimNode_BlendSpacePlayer.cpp b/Engine/Source/Runtime/AnimGraphRuntime/Private/AnimNodes/AnimNode_BlendSpacePlayer.cpp
--- a/Engine/Source/Runtime/AnimGraphRuntime/Private/AnimNodes/AnimNode_BlendSpacePlayer.cpp
+++ b/Engine/Source/Runtime/AnimGraphRuntime/Private/AnimNodes/AnimNode_BlendSpacePlayer.cpp
@@ -96,6 +96,7 @@ void FAnimNode_BlendSpacePlayer::UpdateInternal(const FAnimationUpdateContext& C
 		UE::Anim::FAnimSyncGroupScope& SyncScope = Context.GetMessageChecked<UE::Anim::FAnimSyncGroupScope>();
 
 		FAnimTickRecord TickRecord(BlendSpace, BlendInput, BlendSampleDataCache, BlendFilter, bLoop, PlayRate, Context.GetFinalBlendWeight(), /*inout*/ InternalTimeAccumulator, MarkerTickRecord);
+		TickRecord.RootMotionWeightModifier = Context.GetRootMotionWeightModifier();
 		UE::Anim::FAnimSyncParams SyncParams(GroupName, GroupRole, Method);
 
 		SyncScope.AddTickRecord(TickRecord, SyncParams, UE::Anim::FAnimSyncDebugInfo(Context));
diff --git a/Engine/Source/Runtime/Engine/Classes/Animation/AnimationAsset.h b/Engine/Source/Runtime/Engine/Classes/Animation/AnimationAsset.h
--- a/Engine/Source/Runtime/Engine/Classes/Animation/AnimationAsset.h
+++ b/Engine/Source/Runtime/Engine/Classes/Animation/AnimationAsset.h
@@ -290,9 +290,9 @@ struct FAnimTickRecord
 	class UAnimationAsset* SourceAsset = nullptr;
 
 	float* TimeAccumulator = nullptr;
-	float PlayRateMultiplier = 0.0f;
+	float PlayRateMultiplier = 1.0f;
 	float EffectiveBlendWeight = 0.0f;
-	float RootMotionWeightModifier = 0.0f;
+	float RootMotionWeightModifier = 1.0f;
 
 	bool bLooping = false;
 
