Allow blend list nodes to use property access

Because values are always pins, they dont expose the UI for bindings

diff --git a/Engine/Source/Runtime/AnimGraphRuntime/Public/AnimNodes/AnimNode_BlendListByBool.h b/Engine/Source/Runtime/AnimGraphRuntime/Public/AnimNodes/AnimNode_BlendListByBool.h
--- a/Engine/Source/Runtime/AnimGraphRuntime/Public/AnimNodes/AnimNode_BlendListByBool.h
+++ b/Engine/Source/Runtime/AnimGraphRuntime/Public/AnimNodes/AnimNode_BlendListByBool.h
@@ -14,7 +14,7 @@ struct ANIMGRAPHRUNTIME_API FAnimNode_BlendListByBool : public FAnimNode_BlendLi
 	GENERATED_USTRUCT_BODY()
 public:
 	// Which input should be connected to the output?
-	UPROPERTY(EditAnywhere, BlueprintReadWrite, Category=Runtime, meta=(AlwaysAsPin))
+	UPROPERTY(EditAnywhere, BlueprintReadWrite, Category=Runtime, meta=(PinShownByDefault))
 	bool bActiveValue;
 
 public:	
diff --git a/Engine/Source/Runtime/AnimGraphRuntime/Public/AnimNodes/AnimNode_BlendListByEnum.h b/Engine/Source/Runtime/AnimGraphRuntime/Public/AnimNodes/AnimNode_BlendListByEnum.h
--- a/Engine/Source/Runtime/AnimGraphRuntime/Public/AnimNodes/AnimNode_BlendListByEnum.h
+++ b/Engine/Source/Runtime/AnimGraphRuntime/Public/AnimNodes/AnimNode_BlendListByEnum.h
@@ -18,7 +18,7 @@ public:
 	TArray<int32> EnumToPoseIndex;
 
 	// The currently selected pose (as an enum value)
-	UPROPERTY(EditAnywhere, BlueprintReadWrite, Category=Runtime, meta=(AlwaysAsPin))
+	UPROPERTY(EditAnywhere, BlueprintReadWrite, Category=Runtime, meta=(PinShownByDefault))
 	mutable uint8 ActiveEnumValue;
 
 public:	
diff --git a/Engine/Source/Runtime/AnimGraphRuntime/Public/AnimNodes/AnimNode_BlendListByInt.h b/Engine/Source/Runtime/AnimGraphRuntime/Public/AnimNodes/AnimNode_BlendListByInt.h
--- a/Engine/Source/Runtime/AnimGraphRuntime/Public/AnimNodes/AnimNode_BlendListByInt.h
+++ b/Engine/Source/Runtime/AnimGraphRuntime/Public/AnimNodes/AnimNode_BlendListByInt.h
@@ -13,7 +13,7 @@ struct ANIMGRAPHRUNTIME_API FAnimNode_BlendListByInt : public FAnimNode_BlendLis
 {
 	GENERATED_USTRUCT_BODY()
 public:
-	UPROPERTY(EditAnywhere, BlueprintReadWrite, Category=Runtime, meta=(AlwaysAsPin))
+	UPROPERTY(EditAnywhere, BlueprintReadWrite, Category=Runtime, meta=(PinShownByDefault))
 	int32 ActiveChildIndex;
 
 public:	
