Move CalculateDirection to KismetAnimationLibrary

Function should be available in more than just anim instances, expecially now we want to prefer using property access

diff --git a/Engine/Config/BaseEngine.ini b/Engine/Config/BaseEngine.ini
--- a/Engine/Config/BaseEngine.ini
+++ b/Engine/Config/BaseEngine.ini
@@ -864,6 +864,7 @@ NearClipPlane=10.0
 +ClassRedirects=(OldName="/Script/MeshDescription.MeshDescriptionBulkDataWrapper",NewName="/Script/Engine.StaticMeshDescriptionBulkData")
 +StructRedirects=(OldName="/Script/Engine.FilePath",NewName="/Script/CoreUObject.FilePath")
 +StructRedirects=(OldName="/Script/Engine.DirectoryPath",NewName="/Script/CoreUObject.DirectoryPath")
++FunctionRedirects=(OldName="CalculateDirection",NewName="/Script/AnimGraphRuntime.KismetAnimationLibrary.CalculateDirection")
 
 [CoreUObject.Metadata]
 ; Note: UnrealHeaderTool should be rerun after making changes to MetadataRedirects to catch any keys specified in class headers
diff --git a/Engine/Source/Runtime/AnimGraphRuntime/Private/KismetAnimationLibrary.cpp b/Engine/Source/Runtime/AnimGraphRuntime/Private/KismetAnimationLibrary.cpp
--- a/Engine/Source/Runtime/AnimGraphRuntime/Private/KismetAnimationLibrary.cpp
+++ b/Engine/Source/Runtime/AnimGraphRuntime/Private/KismetAnimationLibrary.cpp
@@ -254,5 +254,32 @@ float UKismetAnimationLibrary::K2_EndProfilingTimer(bool bLog, const FString& Lo
 	return Delta;
 }
 
+float UKismetAnimationLibrary::CalculateDirection(const FVector& Velocity, const FRotator& BaseRotation)
+{
+	if (!Velocity.IsNearlyZero())
+	{
+		FMatrix RotMatrix = FRotationMatrix(BaseRotation);
+		FVector ForwardVector = RotMatrix.GetScaledAxis(EAxis::X);
+		FVector RightVector = RotMatrix.GetScaledAxis(EAxis::Y);
+		FVector NormalizedVel = Velocity.GetSafeNormal2D();
+
+		// get a cos(alpha) of forward vector vs velocity
+		float ForwardCosAngle = FVector::DotProduct(ForwardVector, NormalizedVel);
+		// now get the alpha and convert to degree
+		float ForwardDeltaDegree = FMath::RadiansToDegrees(FMath::Acos(ForwardCosAngle));
+
+		// depending on where right vector is, flip it
+		float RightCosAngle = FVector::DotProduct(RightVector, NormalizedVel);
+		if (RightCosAngle < 0)
+		{
+			ForwardDeltaDegree *= -1;
+		}
+
+		return ForwardDeltaDegree;
+	}
+
+	return 0.f;
+}
+
 #undef LOCTEXT_NAMESPACE
 
diff --git a/Engine/Source/Runtime/AnimGraphRuntime/Public/KismetAnimationLibrary.h b/Engine/Source/Runtime/AnimGraphRuntime/Public/KismetAnimationLibrary.h
--- a/Engine/Source/Runtime/AnimGraphRuntime/Public/KismetAnimationLibrary.h
+++ b/Engine/Source/Runtime/AnimGraphRuntime/Public/KismetAnimationLibrary.h
@@ -215,6 +215,15 @@ class ANIMGRAPHRUNTIME_API UKismetAnimationLibrary : public UBlueprintFunctionLi
 	UFUNCTION(BlueprintCallable, Category = "Utilities|Animation", meta = (DisplayName = "End Profiling Timer"))
 	static float K2_EndProfilingTimer(bool bLog = true, const FString& LogPrefix = "");
 
+	/** 
+	 * Returns degree of the angle between Velocity and Rotation forward vector
+	 * The range of return will be from [-180, 180]. Useful for feeding directional blendspaces.
+	 * @param	Velocity		The velocity to use as direction relative to BaseRotation
+	 * @param	BaseRotation	The base rotation, e.g. of a pawn
+	 */
+	UFUNCTION(BlueprintPure, Category="Utilities|Animation", meta = (BlueprintThreadSafe))
+	static float CalculateDirection(const FVector& Velocity, const FRotator& BaseRotation);
+
 private:
 
 	struct FK2ProfilingTimer
diff --git a/Engine/Source/Runtime/Engine/Classes/Animation/AnimInstance.h b/Engine/Source/Runtime/Engine/Classes/Animation/AnimInstance.h
--- a/Engine/Source/Runtime/Engine/Classes/Animation/AnimInstance.h
+++ b/Engine/Source/Runtime/Engine/Classes/Animation/AnimInstance.h
@@ -1052,11 +1052,7 @@ public:
 	UFUNCTION(BlueprintCallable, Category="Animation")
 	void ClearMorphTargets();
 
-	/** 
-	 * Returns degree of the angle betwee velocity and Rotation forward vector
-	 * The range of return will be from [-180, 180], and this can be used to feed blendspace directional value
-	 */
-	UFUNCTION(BlueprintCallable, Category="Animation")
+	UE_DEPRECATED(5.0, "Please use UKismetAnimationLibrary::CalculateDirection instead")
 	float CalculateDirection(const FVector& Velocity, const FRotator& BaseRotation) const;
 
 	//--- AI communication start ---//
