Fix for bStartFromMatchingPose only working when used within state machines re-init.
#rnx

diff --git a/Engine/Source/Runtime/AnimGraphRuntime/Private/AnimNodes/AnimNode_BlendListBase.cpp b/Engine/Source/Runtime/AnimGraphRuntime/Private/AnimNodes/AnimNode_BlendListBase.cpp
--- a/Engine/Source/Runtime/AnimGraphRuntime/Private/AnimNodes/AnimNode_BlendListBase.cpp
+++ b/Engine/Source/Runtime/AnimGraphRuntime/Private/AnimNodes/AnimNode_BlendListBase.cpp
@@ -181,7 +181,7 @@ void FAnimNode_BlendListBase::Update_AnyThread(const FAnimationUpdateContext& Co
 			// when this flag is true, we'll reinitialize the children
 			if (bResetChildOnActivation)
 			{
-				FAnimationInitializeContext ReinitializeContext(Context.AnimInstanceProxy);
+				FAnimationInitializeContext ReinitializeContext(Context.AnimInstanceProxy, Context.SharedContext);
 
 				// reinitialize
 				BlendPose[ChildIndex].Initialize(ReinitializeContext);
diff --git a/Engine/Source/Runtime/AnimGraphRuntime/Private/AnimNodes/AnimNode_Mirror.cpp b/Engine/Source/Runtime/AnimGraphRuntime/Private/AnimNodes/AnimNode_Mirror.cpp
--- a/Engine/Source/Runtime/AnimGraphRuntime/Private/AnimNodes/AnimNode_Mirror.cpp
+++ b/Engine/Source/Runtime/AnimGraphRuntime/Private/AnimNodes/AnimNode_Mirror.cpp
@@ -100,7 +100,7 @@ void FAnimNode_Mirror::Update_AnyThread(const FAnimationUpdateContext& Context)
 		// Optionally reinitialize the source when the mirror state changes
 		if (bResetChildOnMirrorStateChange)
 		{
-			FAnimationInitializeContext ReinitializeContext(Context.AnimInstanceProxy);
+			FAnimationInitializeContext ReinitializeContext(Context.AnimInstanceProxy, Context.SharedContext);
 			Source.Initialize(ReinitializeContext);
 		}
 	}
diff --git a/Engine/Source/Runtime/AnimGraphRuntime/Private/AnimNodes/AnimNode_TwoWayBlend.cpp b/Engine/Source/Runtime/AnimGraphRuntime/Private/AnimNodes/AnimNode_TwoWayBlend.cpp
--- a/Engine/Source/Runtime/AnimGraphRuntime/Private/AnimNodes/AnimNode_TwoWayBlend.cpp
+++ b/Engine/Source/Runtime/AnimGraphRuntime/Private/AnimNodes/AnimNode_TwoWayBlend.cpp
@@ -63,7 +63,7 @@ void FAnimNode_TwoWayBlend::Update_AnyThread(const FAnimationUpdateContext& Cont
 	{
 		if (bNewAIsRelevant && !bAIsRelevant)
 		{
-			FAnimationInitializeContext ReinitializeContext(Context.AnimInstanceProxy);
+			FAnimationInitializeContext ReinitializeContext(Context.AnimInstanceProxy, Context.SharedContext);
 
 			// reinitialize
 			A.Initialize(ReinitializeContext);
@@ -71,7 +71,7 @@ void FAnimNode_TwoWayBlend::Update_AnyThread(const FAnimationUpdateContext& Cont
 
 		if (bNewBIsRelevant && !bBIsRelevant)
 		{
-			FAnimationInitializeContext ReinitializeContext(Context.AnimInstanceProxy);
+			FAnimationInitializeContext ReinitializeContext(Context.AnimInstanceProxy, Context.SharedContext);
 
 			// reinitialize
 			B.Initialize(ReinitializeContext);
