Anim node data/compiler refactor

diff --git a/Engine/Source/Editor/AnimGraph/Public/PropertyAccessCompilerHandler.h b/Engine/Source/Editor/AnimGraph/Public/PropertyAccessCompilerHandler.h
--- a/Engine/Source/Editor/AnimGraph/Public/PropertyAccessCompilerHandler.h
+++ b/Engine/Source/Editor/AnimGraph/Public/PropertyAccessCompilerHandler.h
@@ -10,10 +10,14 @@
 
 enum class EPropertyAccessBatchType : uint8;
 class IAnimBlueprintGeneratedClassCompiledData;
+class IAnimBlueprintCompilationBracketContext;
 
 // Delegate called when the library is compiled (whether successfully or not)
-DECLARE_MULTICAST_DELEGATE_OneParam(FOnPostLibraryCompiled, IAnimBlueprintGeneratedClassCompiledData& /*OutCompiledData*/)
+DECLARE_MULTICAST_DELEGATE_TwoParams(FOnPostLibraryCompiled, IAnimBlueprintCompilationBracketContext& /*InCompilationContext*/, IAnimBlueprintGeneratedClassCompiledData& /*OutCompiledData*/)
 
+class UE_DEPRECATED(5.0, "FPropertyAccessCompilerHandler is no longer used. Use UAnimBlueprintExtension_PropertyAccess instead") FPropertyAccessCompilerHandler;
+
+PRAGMA_DISABLE_DEPRECATION_WARNINGS
 class FPropertyAccessCompilerHandler : public IAnimBlueprintCompilerHandler
 {
 public:
@@ -29,4 +33,5 @@ public:
 
 	// Maps the initial integer copy handle to a true handle, post compilation
 	virtual int32 MapCopyIndex(int32 InIndex) const = 0;
-};
\ No newline at end of file
+};
+PRAGMA_ENABLE_DEPRECATION_WARNINGS
\ No newline at end of file
diff --git a/Engine/Source/Editor/AnimationBlueprintEditor/AnimationBlueprintEditor.Build.cs b/Engine/Source/Editor/AnimationBlueprintEditor/AnimationBlueprintEditor.Build.cs
--- a/Engine/Source/Editor/AnimationBlueprintEditor/AnimationBlueprintEditor.Build.cs
+++ b/Engine/Source/Editor/AnimationBlueprintEditor/AnimationBlueprintEditor.Build.cs
@@ -42,6 +42,7 @@ public class AnimationBlueprintEditor : ModuleRules
                 "RHI",
                 "KismetCompiler",
 				"ToolMenus",
+				"AnimGraphRuntime",
             }
 		);
 
diff --git a/Engine/Source/Editor/AnimationBlueprintEditor/Private/AnimationBlueprintEditor.cpp b/Engine/Source/Editor/AnimationBlueprintEditor/Private/AnimationBlueprintEditor.cpp
--- a/Engine/Source/Editor/AnimationBlueprintEditor/Private/AnimationBlueprintEditor.cpp
+++ b/Engine/Source/Editor/AnimationBlueprintEditor/Private/AnimationBlueprintEditor.cpp
@@ -758,7 +758,7 @@ void FAnimationBlueprintEditor::OnConvertToSequenceEvaluator()
 			UAnimGraphNode_SequencePlayer* OldNode = Cast<UAnimGraphNode_SequencePlayer>(*NodeIter);
 
 			// see if sequence player
-			if ( OldNode && OldNode->Node.Sequence )
+			if ( OldNode && OldNode->Node.GetSequence() )
 			{
 				UEdGraph* TargetGraph = OldNode->GetGraph();
 				TargetGraph->Modify();
@@ -767,7 +767,7 @@ void FAnimationBlueprintEditor::OnConvertToSequenceEvaluator()
 				// create new evaluator
 				FGraphNodeCreator<UAnimGraphNode_SequenceEvaluator> NodeCreator(*TargetGraph);
 				UAnimGraphNode_SequenceEvaluator* NewNode = NodeCreator.CreateNode();
-				NewNode->Node.Sequence = OldNode->Node.Sequence;
+				NewNode->Node.SetSequence(OldNode->Node.GetSequence());
 				NodeCreator.Finalize();
 
 				// get default data from old node to new node
@@ -807,7 +807,7 @@ void FAnimationBlueprintEditor::OnConvertToSequencePlayer()
 			UAnimGraphNode_SequenceEvaluator* OldNode = Cast<UAnimGraphNode_SequenceEvaluator>(*NodeIter);
 
 			// see if sequence player
-			if ( OldNode && OldNode->Node.Sequence )
+			if ( OldNode && OldNode->Node.GetSequence() )
 			{
 				// convert to sequence player
 				UEdGraph* TargetGraph = OldNode->GetGraph();
@@ -817,7 +817,7 @@ void FAnimationBlueprintEditor::OnConvertToSequencePlayer()
 				// create new player
 				FGraphNodeCreator<UAnimGraphNode_SequencePlayer> NodeCreator(*TargetGraph);
 				UAnimGraphNode_SequencePlayer* NewNode = NodeCreator.CreateNode();
-				NewNode->Node.Sequence = OldNode->Node.Sequence;
+				NewNode->Node.SetSequence(OldNode->Node.GetSequence());
 				NodeCreator.Finalize();
 
 				// get default data from old node to new node
@@ -858,7 +858,7 @@ void FAnimationBlueprintEditor::OnConvertToBlendSpaceEvaluator()
 			UAnimGraphNode_BlendSpacePlayer* OldNode = Cast<UAnimGraphNode_BlendSpacePlayer>(*NodeIter);
 
 			// see if sequence player
-			if ( OldNode && OldNode->Node.BlendSpace )
+			if ( OldNode && OldNode->Node.GetBlendSpace() )
 			{
 				// convert to sequence evaluator
 				UEdGraph* TargetGraph = OldNode->GetGraph();
@@ -868,7 +868,7 @@ void FAnimationBlueprintEditor::OnConvertToBlendSpaceEvaluator()
 				// create new evaluator
 				FGraphNodeCreator<UAnimGraphNode_BlendSpaceEvaluator> NodeCreator(*TargetGraph);
 				UAnimGraphNode_BlendSpaceEvaluator* NewNode = NodeCreator.CreateNode();
-				NewNode->Node.BlendSpace = OldNode->Node.BlendSpace;
+				NewNode->Node.SetBlendSpace(OldNode->Node.GetBlendSpace());
 				NodeCreator.Finalize();
 
 				// get default data from old node to new node
@@ -909,7 +909,7 @@ void FAnimationBlueprintEditor::OnConvertToBlendSpacePlayer()
 			UAnimGraphNode_BlendSpaceEvaluator* OldNode = Cast<UAnimGraphNode_BlendSpaceEvaluator>(*NodeIter);
 
 			// see if sequence player
-			if ( OldNode && OldNode->Node.BlendSpace )
+			if ( OldNode && OldNode->Node.GetBlendSpace() )
 			{
 				// convert to sequence player
 				UEdGraph* TargetGraph = OldNode->GetGraph();
@@ -919,7 +919,7 @@ void FAnimationBlueprintEditor::OnConvertToBlendSpacePlayer()
 				// create new player
 				FGraphNodeCreator<UAnimGraphNode_BlendSpacePlayer> NodeCreator(*TargetGraph);
 				UAnimGraphNode_BlendSpacePlayer* NewNode = NodeCreator.CreateNode();
-				NewNode->Node.BlendSpace = OldNode->Node.BlendSpace;
+				NewNode->Node.SetBlendSpace(OldNode->Node.GetBlendSpace());
 				NodeCreator.Finalize();
 
 				// get default data from old node to new node
@@ -960,7 +960,7 @@ void FAnimationBlueprintEditor::OnConvertToBlendSpaceGraph()
 			UAnimGraphNode_BlendSpacePlayer* OldNode = Cast<UAnimGraphNode_BlendSpacePlayer>(*NodeIter);
 
 			// see if sequence player
-			if (OldNode && OldNode->Node.BlendSpace)
+			if (OldNode && OldNode->Node.GetBlendSpace())
 			{
 				// convert to sequence player
 				UEdGraph* TargetGraph = OldNode->GetGraph();
@@ -970,11 +970,11 @@ void FAnimationBlueprintEditor::OnConvertToBlendSpaceGraph()
 				// create new player
 				FGraphNodeCreator<UAnimGraphNode_BlendSpaceGraph> NodeCreator(*TargetGraph);
 				UAnimGraphNode_BlendSpaceGraph* NewNode = NodeCreator.CreateNode();
-				NewNode->SetupFromAsset(OldNode->Node.BlendSpace, false);
-				if(OldNode->SyncGroup.GroupName != NAME_None && OldNode->SyncGroup.Method == EAnimSyncMethod::SyncGroup)
+				if(OldNode->Node.GetGroupName() != NAME_None && OldNode->Node.GetGroupMethod() == EAnimSyncMethod::SyncGroup)
 				{
-					NewNode->SetSyncGroupName(OldNode->SyncGroup.GroupName);
+					NewNode->SetSyncGroupName(OldNode->Node.GetGroupName());
 				}
+				NewNode->SetupFromAsset(OldNode->Node.GetBlendSpace(), false);
 				NodeCreator.Finalize();
 
 				// get default data from old node to new node
@@ -1116,7 +1116,7 @@ void FAnimationBlueprintEditor::OnConvertToAimOffsetLookAt()
 			UAnimGraphNode_RotationOffsetBlendSpace* OldNode = Cast<UAnimGraphNode_RotationOffsetBlendSpace>(*NodeIter);
 
 			// see if sequence player
-			if (OldNode && OldNode->Node.BlendSpace)
+			if (OldNode && OldNode->Node.GetBlendSpace())
 			{
 				// convert to sequence evaluator
 				UEdGraph* TargetGraph = OldNode->GetGraph();
@@ -1126,7 +1126,7 @@ void FAnimationBlueprintEditor::OnConvertToAimOffsetLookAt()
 				// create new evaluator
 				FGraphNodeCreator<UAnimGraphNode_AimOffsetLookAt> NodeCreator(*TargetGraph);
 				UAnimGraphNode_AimOffsetLookAt* NewNode = NodeCreator.CreateNode();
-				NewNode->Node.BlendSpace = OldNode->Node.BlendSpace;
+				NewNode->Node.SetBlendSpace(OldNode->Node.GetBlendSpace());
 				NodeCreator.Finalize();
 
 				// get default data from old node to new node
@@ -1170,7 +1170,7 @@ void FAnimationBlueprintEditor::OnConvertToAimOffsetSimple()
 			UAnimGraphNode_AimOffsetLookAt* OldNode = Cast<UAnimGraphNode_AimOffsetLookAt>(*NodeIter);
 
 			// see if sequence player
-			if (OldNode && OldNode->Node.BlendSpace)
+			if (OldNode && OldNode->Node.GetBlendSpace())
 			{
 				// convert to sequence player
 				UEdGraph* TargetGraph = OldNode->GetGraph();
@@ -1180,7 +1180,7 @@ void FAnimationBlueprintEditor::OnConvertToAimOffsetSimple()
 				// create new player
 				FGraphNodeCreator<UAnimGraphNode_RotationOffsetBlendSpace> NodeCreator(*TargetGraph);
 				UAnimGraphNode_RotationOffsetBlendSpace* NewNode = NodeCreator.CreateNode();
-				NewNode->Node.BlendSpace = OldNode->Node.BlendSpace;
+				NewNode->Node.SetBlendSpace(OldNode->Node.GetBlendSpace());
 				NodeCreator.Finalize();
 
 				// get default data from old node to new node
@@ -1222,7 +1222,7 @@ void FAnimationBlueprintEditor::OnConvertToAimOffsetGraph()
 			UAnimGraphNode_RotationOffsetBlendSpace* OldNode = Cast<UAnimGraphNode_RotationOffsetBlendSpace>(*NodeIter);
 
 			// see if sequence player
-			if (OldNode && OldNode->Node.BlendSpace)
+			if (OldNode && OldNode->Node.GetBlendSpace())
 			{
 				// convert to sequence player
 				UEdGraph* TargetGraph = OldNode->GetGraph();
@@ -1232,11 +1232,11 @@ void FAnimationBlueprintEditor::OnConvertToAimOffsetGraph()
 				// create new player
 				FGraphNodeCreator<UAnimGraphNode_RotationOffsetBlendSpaceGraph> NodeCreator(*TargetGraph);
 				UAnimGraphNode_RotationOffsetBlendSpaceGraph* NewNode = NodeCreator.CreateNode();
-				NewNode->SetupFromAsset(OldNode->Node.BlendSpace, false);
-				if(OldNode->SyncGroup.GroupName != NAME_None && OldNode->SyncGroup.Method == EAnimSyncMethod::SyncGroup)
+				if(OldNode->Node.GetGroupName() != NAME_None && OldNode->Node.GetGroupMethod() == EAnimSyncMethod::SyncGroup)
 				{
-					NewNode->SetSyncGroupName(OldNode->SyncGroup.GroupName);
+					NewNode->SetSyncGroupName(OldNode->Node.GetGroupName());
 				}
+				NewNode->SetupFromAsset(OldNode->Node.GetBlendSpace(), false);
 				NodeCreator.Finalize();
 
 				// get default data from old node to new node
diff --git a/Engine/Source/Editor/AnimationBlueprintEditor/Private/AnimationNodes/SGraphNodeBlendSpacePlayer.cpp b/Engine/Source/Editor/AnimationBlueprintEditor/Private/AnimationNodes/SGraphNodeBlendSpacePlayer.cpp
--- a/Engine/Source/Editor/AnimationBlueprintEditor/Private/AnimationNodes/SGraphNodeBlendSpacePlayer.cpp
+++ b/Engine/Source/Editor/AnimationBlueprintEditor/Private/AnimationNodes/SGraphNodeBlendSpacePlayer.cpp
@@ -48,7 +48,7 @@ void SGraphNodeBlendSpacePlayer::UpdateGraphSyncLabel()
 			{
 				if (UObject* ActiveObject = AnimBlueprint->GetObjectBeingDebugged())
 				{
-					if(VisualBlendSpacePlayer->SyncGroup.Method == EAnimSyncMethod::Graph)
+					if(VisualBlendSpacePlayer->Node.GetGroupMethod() == EAnimSyncMethod::Graph)
 					{
 						int32 NodeIndex = GeneratedClass->GetNodeIndexFromGuid(VisualBlendSpacePlayer->NodeGuid);
 						if(NodeIndex != INDEX_NONE)
diff --git a/Engine/Source/Editor/AnimationBlueprintEditor/Private/AnimationNodes/SGraphNodeSequencePlayer.cpp b/Engine/Source/Editor/AnimationBlueprintEditor/Private/AnimationNodes/SGraphNodeSequencePlayer.cpp
--- a/Engine/Source/Editor/AnimationBlueprintEditor/Private/AnimationNodes/SGraphNodeSequencePlayer.cpp
+++ b/Engine/Source/Editor/AnimationBlueprintEditor/Private/AnimationNodes/SGraphNodeSequencePlayer.cpp
@@ -176,9 +176,10 @@ void SGraphNodeSequencePlayer::SetSequencePositionRatio(float NewRatio)
 {
 	if(FAnimNode_SequencePlayer* SequencePlayer = GetSequencePlayer())
 	{
-		if (SequencePlayer->Sequence != NULL)
+		UAnimSequenceBase* Sequence = SequencePlayer->GetSequence();
+		if (Sequence != NULL)
 		{
-			const float NewTime = NewRatio * SequencePlayer->Sequence->GetPlayLength();
+			const float NewTime = NewRatio * Sequence->GetPlayLength();
 			SequencePlayer->SetAccumulatedTime(NewTime);
 		}
 	}
@@ -196,7 +197,7 @@ void SGraphNodeSequencePlayer::UpdateGraphSyncLabel()
 			{
 				if (UObject* ActiveObject = AnimBlueprint->GetObjectBeingDebugged())
 				{
-					if(VisualSequencePlayer->SyncGroup.Method == EAnimSyncMethod::Graph)
+					if(VisualSequencePlayer->Node.GetGroupMethod() == EAnimSyncMethod::Graph)
 					{
 						int32 NodeIndex = GeneratedClass->GetNodeIndexFromGuid(VisualSequencePlayer->NodeGuid);
 						if(NodeIndex != INDEX_NONE)
diff --git a/Engine/Source/Editor/UnrealEd/Private/Kismet2/BlueprintEditorUtils.cpp b/Engine/Source/Editor/UnrealEd/Private/Kismet2/BlueprintEditorUtils.cpp
--- a/Engine/Source/Editor/UnrealEd/Private/Kismet2/BlueprintEditorUtils.cpp
+++ b/Engine/Source/Editor/UnrealEd/Private/Kismet2/BlueprintEditorUtils.cpp
@@ -4224,7 +4224,7 @@ void FBlueprintEditorUtils::SetAnimationGraphLayerGroup(UEdGraph* InGraph, const
 	{
 		const FName NewGroup = InGroupName.IsEmpty() ? NAME_None : FName(*InGroupName.ToString());
 		UAnimGraphNode_Root* Root = GetAnimGraphRoot(InGraph);
-		if(NewGroup != Root->Node.Group)
+		if(NewGroup != Root->Node.GetGroup())
 		{
 			FScopedTransaction Transaction(LOCTEXT("SetAnimationGraphLayerGroup", "Set Group"));
 
@@ -4243,7 +4243,7 @@ void FBlueprintEditorUtils::SetAnimationGraphLayerGroup(UEdGraph* InGraph, const
 			}
 
 		
-			Root->Node.Group = NewGroup;
+			Root->Node.SetGroup(NewGroup);
 
 			if (Function)
 			{
diff --git a/Engine/Source/Editor/UnrealEd/Public/IPropertyAccessCompiler.h b/Engine/Source/Editor/UnrealEd/Public/IPropertyAccessCompiler.h
--- a/Engine/Source/Editor/UnrealEd/Public/IPropertyAccessCompiler.h
+++ b/Engine/Source/Editor/UnrealEd/Public/IPropertyAccessCompiler.h
@@ -5,6 +5,7 @@
 #include "CoreMinimal.h"
 #include "Containers/ArrayView.h"
 #include "UObject/ObjectMacros.h"
+#include "Features/IModularFeature.h"
 
 class UObject;
 struct FPropertyAccessLibrary;
@@ -27,7 +28,10 @@ public:
 	virtual ~IPropertyAccessLibraryCompiler() {}
 
 	// Begin compilation - reset the library to its default state
-	virtual void BeginCompilation(const UClass* InClass) = 0;
+	virtual void BeginCompilation() = 0;
+	
+	UE_DEPRECATED(5.0, "Please use BeginCompilation without a class argument")
+	virtual void BeginCompilation(const UClass* InClass) {}
 
 	// Add a copy to the property access library we are compiling
 	// @return an integer handle to the pending copy. This can be resolved to a true copy index by calling MapCopyIndex
diff --git a/Engine/Source/Editor/UnrealEd/Public/IPropertyAccessEditor.h b/Engine/Source/Editor/UnrealEd/Public/IPropertyAccessEditor.h
--- a/Engine/Source/Editor/UnrealEd/Public/IPropertyAccessEditor.h
+++ b/Engine/Source/Editor/UnrealEd/Public/IPropertyAccessEditor.h
@@ -16,6 +16,8 @@ class FExtender;
 class SWidget;
 struct FSlateBrush;
 struct FEdGraphPinType;
+class IPropertyAccessLibraryCompiler;
+struct FPropertyAccessLibrary;
 
 /** An element in a binding chain */
 struct FBindingChainElement
@@ -206,6 +208,21 @@ enum class EPropertyAccessCompatibility
 	Promotable,
 };
 
+// Context used to create a property access library compiler
+struct FPropertyAccessLibraryCompilerArgs
+{
+	FPropertyAccessLibraryCompilerArgs(FPropertyAccessLibrary& InLibrary, const UClass* InClassContext)
+		: Library(InLibrary)
+		, ClassContext(InClassContext)
+	{}
+
+	// The library that will be built
+	FPropertyAccessLibrary& Library;
+
+	// The class that provides a root context for the library to be built in
+	const UClass* ClassContext;
+};
+
 /** Editor support for property access system */
 class IPropertyAccessEditor : public IModularFeature
 {
@@ -239,4 +256,7 @@ public:
 
 	// Makes a string path from a binding chain
 	virtual void MakeStringPath(const TArray<FBindingChainElement>& InBindingChain, TArray<FString>& OutStringPath) const = 0;
+
+	// Make a property access library compiler, used for building a FPropertyAccessLibrary
+	virtual TUniquePtr<IPropertyAccessLibraryCompiler> MakePropertyAccessCompiler(const FPropertyAccessLibraryCompilerArgs& InArgs) const = 0;
 };
\ No newline at end of file
diff --git a/Engine/Source/Runtime/AnimGraphRuntime/Private/AnimGraphRuntimeTrace.cpp b/Engine/Source/Runtime/AnimGraphRuntime/Private/AnimGraphRuntimeTrace.cpp
--- a/Engine/Source/Runtime/AnimGraphRuntime/Private/AnimGraphRuntimeTrace.cpp
+++ b/Engine/Source/Runtime/AnimGraphRuntime/Private/AnimGraphRuntimeTrace.cpp
@@ -34,19 +34,22 @@ void FAnimGraphRuntimeTrace::OutputBlendSpacePlayer(const FAnimationBaseContext&
 	check(InContext.AnimInstanceProxy);
 
 	TRACE_OBJECT(InContext.AnimInstanceProxy->GetAnimInstanceObject());
-	TRACE_OBJECT(InNode.BlendSpace);
+	TRACE_OBJECT(InNode.GetBlendSpace());
+
+	FVector SampleCoordinates = InNode.GetPosition();
+	FVector FilteredPosition = InNode.GetFilteredPosition();
 
 	UE_TRACE_LOG(Animation, BlendSpacePlayer, AnimationChannel)
 		<< BlendSpacePlayer.Cycle(FPlatformTime::Cycles64())
 		<< BlendSpacePlayer.AnimInstanceId(FObjectTrace::GetObjectId(InContext.AnimInstanceProxy->GetAnimInstanceObject()))
-		<< BlendSpacePlayer.BlendSpaceId(FObjectTrace::GetObjectId(InNode.BlendSpace))
+		<< BlendSpacePlayer.BlendSpaceId(FObjectTrace::GetObjectId(InNode.GetBlendSpace()))
 		<< BlendSpacePlayer.NodeId(InContext.GetCurrentNodeId())
-		<< BlendSpacePlayer.PositionX(InNode.X)
-		<< BlendSpacePlayer.PositionY(InNode.Y)
-		<< BlendSpacePlayer.PositionZ(InNode.Z)
-		<< BlendSpacePlayer.FilteredPositionX(InNode.GetFilteredPosition().X)
-		<< BlendSpacePlayer.FilteredPositionY(InNode.GetFilteredPosition().Y)
-		<< BlendSpacePlayer.FilteredPositionZ(InNode.GetFilteredPosition().Z);
+		<< BlendSpacePlayer.PositionX(SampleCoordinates.X)
+		<< BlendSpacePlayer.PositionY(SampleCoordinates.Y)
+		<< BlendSpacePlayer.PositionZ(SampleCoordinates.Z)
+		<< BlendSpacePlayer.FilteredPositionX(FilteredPosition.X)
+		<< BlendSpacePlayer.FilteredPositionY(FilteredPosition.Y)
+		<< BlendSpacePlayer.FilteredPositionZ(FilteredPosition.Z);
 }
 
 void FAnimGraphRuntimeTrace::OutputBlendSpace(const FAnimationBaseContext& InContext, const FAnimNode_BlendSpaceGraphBase& InNode)
diff --git a/Engine/Source/Runtime/AnimGraphRuntime/Private/AnimNodes/AnimNode_AimOffsetLookAt.cpp b/Engine/Source/Runtime/AnimGraphRuntime/Private/AnimNodes/AnimNode_AimOffsetLookAt.cpp
--- a/Engine/Source/Runtime/AnimGraphRuntime/Private/AnimNodes/AnimNode_AimOffsetLookAt.cpp
+++ b/Engine/Source/Runtime/AnimGraphRuntime/Private/AnimNodes/AnimNode_AimOffsetLookAt.cpp
@@ -103,10 +103,10 @@ void FAnimNode_AimOffsetLookAt::Evaluate_AnyThread(FPoseContext& Context)
 
 void FAnimNode_AimOffsetLookAt::UpdateFromLookAtTarget(FPoseContext& LocalPoseContext)
 {
-	FVector BlendInput(X, Y, Z);
+	UBlendSpace* CurrentBlendSpace = GetBlendSpace();
 
 	const FBoneContainer& RequiredBones = LocalPoseContext.Pose.GetBoneContainer();
-	if (BlendSpace && SocketBoneReference.IsValidToEvaluate(RequiredBones))
+	if (CurrentBlendSpace && SocketBoneReference.IsValidToEvaluate(RequiredBones))
 	{
 		FCSPose<FCompactPose> GlobalPose;
 		GlobalPose.InitPose(LocalPoseContext.Pose);
@@ -140,8 +140,8 @@ void FAnimNode_AimOffsetLookAt::UpdateFromLookAtTarget(FPoseContext& LocalPoseCo
 
 		const FVector2D CurrentCoords = FMath::GetAzimuthAndElevation(CurrentDirection, AxisX, AxisY, AxisZ);
 		const FVector2D TargetCoords = FMath::GetAzimuthAndElevation(DirectionToTarget, AxisX, AxisY, AxisZ);
-		BlendInput.X = FRotator::NormalizeAxis(FMath::RadiansToDegrees(TargetCoords.X - CurrentCoords.X));
-		BlendInput.Y = FRotator::NormalizeAxis(FMath::RadiansToDegrees(TargetCoords.Y - CurrentCoords.Y));
+		CurrentBlendInput.X = FRotator::NormalizeAxis(FMath::RadiansToDegrees(TargetCoords.X - CurrentCoords.X));
+		CurrentBlendInput.Y = FRotator::NormalizeAxis(FMath::RadiansToDegrees(TargetCoords.Y - CurrentCoords.Y));
 
 #if ENABLE_DRAW_DEBUG
 		if (CVarAimOffsetLookAtDebug.GetValueOnAnyThread() == 1)
@@ -155,21 +155,17 @@ void FAnimNode_AimOffsetLookAt::UpdateFromLookAtTarget(FPoseContext& LocalPoseCo
 				, FMath::RadiansToDegrees(CurrentCoords.Y)
 				, FMath::RadiansToDegrees(TargetCoords.X)
 				, FMath::RadiansToDegrees(TargetCoords.Y)
-				, BlendInput.X
-				, BlendInput.Y);
+				, CurrentBlendInput.X
+				, CurrentBlendInput.Y);
 			AnimProxy->AnimDrawDebugOnScreenMessage(DebugString, FColor::Red);
 		}
 #endif // ENABLE_DRAW_DEBUG
 	}
 
-	// Set X and Y, so ticking next frame is based on correct weights.
-	X = BlendInput.X;
-	Y = BlendInput.Y;
-
 	// Generate BlendSampleDataCache from inputs.
-	if (BlendSpace)
+	if (CurrentBlendSpace)
 	{
-		BlendSpace->GetSamplesFromBlendInput(BlendInput, BlendSampleDataCache, CachedTriangulationIndex, true);
+		CurrentBlendSpace->GetSamplesFromBlendInput(CurrentBlendInput, BlendSampleDataCache, CachedTriangulationIndex, true);
 	}
 }
 
@@ -184,6 +180,12 @@ void FAnimNode_AimOffsetLookAt::GatherDebugData(FNodeDebugData& DebugData)
 	BasePose.GatherDebugData(DebugData);
 }
 
+FVector FAnimNode_AimOffsetLookAt::GetPosition() const
+{
+	// Use our calculated coordinates rather than the folded values
+	return CurrentBlendInput;
+}
+
 FAnimNode_AimOffsetLookAt::FAnimNode_AimOffsetLookAt()
 	: SocketLocalTransform(FTransform::Identity)
 	, PivotSocketLocalTransform(FTransform::Identity)
@@ -193,6 +195,7 @@ FAnimNode_AimOffsetLookAt::FAnimNode_AimOffsetLookAt()
 	, LookAtLocation(ForceInitToZero)
 	, SocketAxis(1.0f, 0.0f, 0.0f)
 	, Alpha(1.f)
+	, CurrentBlendInput(FVector::ZeroVector)
 	, bIsLODEnabled(false)
 {
 }
diff --git a/Engine/Source/Runtime/AnimGraphRuntime/Private/AnimNodes/AnimNode_BlendListBase.cpp b/Engine/Source/Runtime/AnimGraphRuntime/Private/AnimNodes/AnimNode_BlendListBase.cpp
--- a/Engine/Source/Runtime/AnimGraphRuntime/Private/AnimNodes/AnimNode_BlendListBase.cpp
+++ b/Engine/Source/Runtime/AnimGraphRuntime/Private/AnimNodes/AnimNode_BlendListBase.cpp
@@ -16,7 +16,8 @@ void FAnimNode_BlendListBase::Initialize_AnyThread(const FAnimationInitializeCon
 	FAnimNode_Base::Initialize_AnyThread(Context);
 
 	const int32 NumPoses = BlendPose.Num();
-	checkSlow(BlendTime.Num() == NumPoses);
+	const TArray<float>& CurrentBlendTimes = GetBlendTimes();
+	checkSlow(CurrentBlendTimes.Num() == NumPoses);
 
 	BlendWeights.Reset(NumPoses);
 	PosesToEvaluate.Reset(NumPoses);
@@ -40,7 +41,8 @@ void FAnimNode_BlendListBase::Initialize_AnyThread(const FAnimationInitializeCon
 	Blends.Empty(NumPoses);
 	Blends.AddZeroed(NumPoses);
 
-	if (BlendProfile)
+	UBlendProfile* CurrentBlendProfile = GetBlendProfile();
+	if (CurrentBlendProfile)
 	{
 		BlendStartAlphas.Empty(NumPoses);
 		BlendStartAlphas.AddZeroed(NumPoses);
@@ -48,22 +50,24 @@ void FAnimNode_BlendListBase::Initialize_AnyThread(const FAnimationInitializeCon
 
 	LastActiveChildIndex = INDEX_NONE;
 
+	EAlphaBlendOption CurrentBlendType = GetBlendType();
+	UCurveFloat* CurrentCustomBlendCurve = GetCustomBlendCurve();
 	for(int32 i = 0 ; i < Blends.Num() ; ++i)
 	{
 		FAlphaBlend& Blend = Blends[i];
 
 		Blend.SetBlendTime(0.0f);
-		Blend.SetBlendOption(BlendType);
-		Blend.SetCustomCurve(CustomBlendCurve);
+		Blend.SetBlendOption(CurrentBlendType);
+		Blend.SetCustomCurve(CurrentCustomBlendCurve);
 
-		if (BlendProfile)
+		if (CurrentBlendProfile)
 		{
 			BlendStartAlphas[i] = 0.0f;
 		}
 	}
 	Blends[0].SetAlpha(1.0f);
 
-	if(BlendProfile)
+	if(CurrentBlendProfile)
 	{
 		BlendStartAlphas[0] = 1.0f;
 
@@ -75,7 +79,7 @@ void FAnimNode_BlendListBase::Initialize_AnyThread(const FAnimationInitializeCon
 		{
 			FBlendSampleData& SampleData = PerBoneSampleData[Idx];
 			SampleData.SampleDataIndex = Idx;
-			SampleData.PerBoneBlendData.AddZeroed(BlendProfile->GetNumBlendEntries());
+			SampleData.PerBoneBlendData.AddZeroed(CurrentBlendProfile->GetNumBlendEntries());
 		}
 	}
 }
@@ -95,12 +99,15 @@ void FAnimNode_BlendListBase::Update_AnyThread(const FAnimationUpdateContext& Co
 	GetEvaluateGraphExposedInputs().Execute(Context);
 
 	const int32 NumPoses = BlendPose.Num();
-	checkSlow((BlendTime.Num() == NumPoses) && (BlendWeights.Num() == NumPoses));
+	const TArray<float>& CurrentBlendTimes = GetBlendTimes();
+	checkSlow((CurrentBlendTimes.Num() == NumPoses) && (BlendWeights.Num() == NumPoses));
 
 	PosesToEvaluate.Empty(NumPoses);
 
 	if (NumPoses > 0)
 	{
+		UBlendProfile* CurrentBlendProfile = GetBlendProfile();
+		
 		// Handle a change in the active child index; adjusting the target weights
 		const int32 ChildIndex = GetActiveChildIndex();
 		
@@ -120,12 +127,12 @@ void FAnimNode_BlendListBase::Update_AnyThread(const FAnimationUpdateContext& Co
 			{
 				RemainingBlendTime = 0.0f;
 			}
-			else if (TransitionType == EBlendListTransitionType::Inertialization)
+			else if (GetTransitionType() == EBlendListTransitionType::Inertialization)
 			{
 				UE::Anim::IInertializationRequester* InertializationRequester = Context.GetMessage<UE::Anim::IInertializationRequester>();
 				if (InertializationRequester)
 				{
-					InertializationRequester->RequestInertialization(BlendTime[ChildIndex]);
+					InertializationRequester->RequestInertialization(CurrentBlendTimes[ChildIndex]);
 					InertializationRequester->AddDebugRecord(*Context.AnimInstanceProxy, Context.GetCurrentNodeId());
 				}
 				else
@@ -137,7 +144,7 @@ void FAnimNode_BlendListBase::Update_AnyThread(const FAnimationUpdateContext& Co
 			}
 			else
 			{
-				RemainingBlendTime = BlendTime[ChildIndex] * WeightDifference;
+				RemainingBlendTime = CurrentBlendTimes[ChildIndex] * WeightDifference;
 			}
 
 			for (int32 i = 0; i < RemainingBlendTimes.Num(); ++i)
@@ -161,7 +168,7 @@ void FAnimNode_BlendListBase::Update_AnyThread(const FAnimationUpdateContext& Co
 				{
 					Blend.SetValueRange(BlendWeights[i], 1.0f);
 
-					if (BlendProfile)
+					if (CurrentBlendProfile)
 					{
 						Blend.ResetAlpha();
 						BlendStartAlphas[i] = Blend.GetAlpha();
@@ -172,14 +179,14 @@ void FAnimNode_BlendListBase::Update_AnyThread(const FAnimationUpdateContext& Co
 					Blend.SetValueRange(BlendWeights[i], 0.0f);
 				}
 
-				if (BlendProfile)
+				if (CurrentBlendProfile)
 				{
 					BlendStartAlphas[i] = Blend.GetAlpha();
 				}
 			}
 
 			// when this flag is true, we'll reinitialize the children
-			if (bResetChildOnActivation)
+			if (GetResetChildOnActivation())
 			{
 				FAnimationInitializeContext ReinitializeContext(Context.AnimInstanceProxy, Context.SharedContext);
 
@@ -226,13 +233,13 @@ void FAnimNode_BlendListBase::Update_AnyThread(const FAnimationUpdateContext& Co
 		}
 
 		// If we're using a blend profile, extract the scales and build blend sample data
-		if (BlendProfile)
+		if (CurrentBlendProfile)
 		{
 			for(int32 i = 0; i < BlendPose.Num(); ++i)
 			{
 				FBlendSampleData& PoseSampleData = PerBoneSampleData[i];
 				PoseSampleData.TotalWeight = BlendWeights[i];
-				BlendProfile->UpdateBoneWeights(PoseSampleData, Blends[i], BlendStartAlphas[i], BlendWeights[i]);
+				CurrentBlendProfile->UpdateBoneWeights(PoseSampleData, Blends[i], BlendStartAlphas[i], BlendWeights[i]);
 			}
 
 			FBlendSampleData::NormalizeDataWeight(PerBoneSampleData);
@@ -243,7 +250,7 @@ void FAnimNode_BlendListBase::Update_AnyThread(const FAnimationUpdateContext& Co
 	const int32 ChildIndex = GetActiveChildIndex();
 	TRACE_ANIM_NODE_VALUE(Context, TEXT("Active Index"), ChildIndex);
 	TRACE_ANIM_NODE_VALUE(Context, TEXT("Active Weight"), BlendWeights[ChildIndex]);
-	TRACE_ANIM_NODE_VALUE(Context, TEXT("Active Blend Time"), BlendTime[ChildIndex]);
+	TRACE_ANIM_NODE_VALUE(Context, TEXT("Active Blend Time"), CurrentBlendTimes[ChildIndex]);
 #endif
 }
 
@@ -283,9 +290,10 @@ void FAnimNode_BlendListBase::Evaluate_AnyThread(FPoseContext& Output)
 		FAnimationPoseData OutAnimationPoseData(Output);
 
 		// Use the calculated blend sample data if we're blending per-bone
-		if (BlendProfile)
+		UBlendProfile* CurrentBlendProfile = GetBlendProfile();
+		if (CurrentBlendProfile)
 		{
-			FAnimationRuntime::BlendPosesTogetherPerBone(FilteredPoses, FilteredCurve, FilteredAttributes, BlendProfile, PerBoneSampleData, PosesToEvaluate, OutAnimationPoseData);
+			FAnimationRuntime::BlendPosesTogetherPerBone(FilteredPoses, FilteredCurve, FilteredAttributes, CurrentBlendProfile, PerBoneSampleData, PosesToEvaluate, OutAnimationPoseData);
 		}
 		else
 		{
@@ -305,7 +313,8 @@ void FAnimNode_BlendListBase::GatherDebugData(FNodeDebugData& DebugData)
 	const int32 ChildIndex = GetActiveChildIndex();
 
 	FString DebugLine = GetNodeName(DebugData);
-	DebugLine += FString::Printf(TEXT("(Active: (%i/%i) Weight: %.1f%% Time %.3f)"), ChildIndex+1, NumPoses, BlendWeights[ChildIndex]*100.f, BlendTime[ChildIndex]);
+	const TArray<float>& CurrentBlendTimes = GetBlendTimes();
+	DebugLine += FString::Printf(TEXT("(Active: (%i/%i) Weight: %.1f%% Time %.3f)"), ChildIndex+1, NumPoses, BlendWeights[ChildIndex]*100.f, CurrentBlendTimes[ChildIndex]);
 
 	DebugData.AddDebugItem(DebugLine);
 	
@@ -314,3 +323,33 @@ void FAnimNode_BlendListBase::GatherDebugData(FNodeDebugData& DebugData)
 		BlendPose[Pose].GatherDebugData(DebugData.BranchFlow(BlendWeights[Pose]));
 	}
 }
+
+const TArray<float>& FAnimNode_BlendListBase::GetBlendTimes() const
+{
+	return GET_ANIM_NODE_DATA(TArray<float>, BlendTime);
+}
+
+EBlendListTransitionType FAnimNode_BlendListBase::GetTransitionType() const
+{
+	return GET_ANIM_NODE_DATA(EBlendListTransitionType, TransitionType);
+}
+
+EAlphaBlendOption FAnimNode_BlendListBase::GetBlendType() const
+{
+	return GET_ANIM_NODE_DATA(EAlphaBlendOption, BlendType);
+}
+
+bool FAnimNode_BlendListBase::GetResetChildOnActivation() const
+{
+	return GET_ANIM_NODE_DATA(bool, bResetChildOnActivation);
+}
+
+UCurveFloat* FAnimNode_BlendListBase::GetCustomBlendCurve() const
+{
+	return GET_ANIM_NODE_DATA(TObjectPtr<UCurveFloat>, CustomBlendCurve);
+}
+
+UBlendProfile* FAnimNode_BlendListBase::GetBlendProfile() const
+{
+	return GET_ANIM_NODE_DATA(TObjectPtr<UBlendProfile>, BlendProfile);
+}
\ No newline at end of file
diff --git a/Engine/Source/Runtime/AnimGraphRuntime/Private/AnimNodes/AnimNode_BlendListByBool.cpp b/Engine/Source/Runtime/AnimGraphRuntime/Private/AnimNodes/AnimNode_BlendListByBool.cpp
--- a/Engine/Source/Runtime/AnimGraphRuntime/Private/AnimNodes/AnimNode_BlendListByBool.cpp
+++ b/Engine/Source/Runtime/AnimGraphRuntime/Private/AnimNodes/AnimNode_BlendListByBool.cpp
@@ -8,5 +8,10 @@
 int32 FAnimNode_BlendListByBool::GetActiveChildIndex()
 {
 	// Note: Intentionally flipped boolean sense (the true input is #0, and the false input is #1)
-	return bActiveValue ? 0 : 1;
+	return GetActiveValue() ? 0 : 1;
 }
+
+bool FAnimNode_BlendListByBool::GetActiveValue() const
+{
+	return GET_ANIM_NODE_DATA(bool, bActiveValue);
+}
\ No newline at end of file
diff --git a/Engine/Source/Runtime/AnimGraphRuntime/Private/AnimNodes/AnimNode_BlendListByEnum.cpp b/Engine/Source/Runtime/AnimGraphRuntime/Private/AnimNodes/AnimNode_BlendListByEnum.cpp
--- a/Engine/Source/Runtime/AnimGraphRuntime/Private/AnimNodes/AnimNode_BlendListByEnum.cpp
+++ b/Engine/Source/Runtime/AnimGraphRuntime/Private/AnimNodes/AnimNode_BlendListByEnum.cpp
@@ -7,9 +7,11 @@
 
 int32 FAnimNode_BlendListByEnum::GetActiveChildIndex()
 {
-	if (EnumToPoseIndex.IsValidIndex(ActiveEnumValue))
+	uint8 CurrentActiveEnumValue = GetActiveEnumValue();
+	const TArray<int32>& CurrentEnumToPoseIndex = GetEnumToPoseIndex();
+	if (CurrentEnumToPoseIndex.IsValidIndex(CurrentActiveEnumValue))
 	{
-		return EnumToPoseIndex[ActiveEnumValue];
+		return CurrentEnumToPoseIndex[CurrentActiveEnumValue];
 	}
 	else
 	{
@@ -17,3 +19,19 @@ int32 FAnimNode_BlendListByEnum::GetActiveChildIndex()
 	}
 }
 
+#if WITH_EDITORONLY_DATA
+void FAnimNode_BlendListByEnum::SetEnumToPoseIndex(const TArray<int32>& InEnumToPoseIndex)
+{
+	EnumToPoseIndex = InEnumToPoseIndex;
+}
+#endif
+
+const TArray<int32>& FAnimNode_BlendListByEnum::GetEnumToPoseIndex() const
+{
+	return GET_ANIM_NODE_DATA(TArray<int32>, EnumToPoseIndex);
+}
+
+uint8 FAnimNode_BlendListByEnum::GetActiveEnumValue() const
+{
+	return GET_ANIM_NODE_DATA(uint8, ActiveEnumValue);
+}
\ No newline at end of file
diff --git a/Engine/Source/Runtime/AnimGraphRuntime/Private/AnimNodes/AnimNode_BlendListByInt.cpp b/Engine/Source/Runtime/AnimGraphRuntime/Private/AnimNodes/AnimNode_BlendListByInt.cpp
--- a/Engine/Source/Runtime/AnimGraphRuntime/Private/AnimNodes/AnimNode_BlendListByInt.cpp
+++ b/Engine/Source/Runtime/AnimGraphRuntime/Private/AnimNodes/AnimNode_BlendListByInt.cpp
@@ -7,6 +7,7 @@
 
 int32 FAnimNode_BlendListByInt::GetActiveChildIndex()
 {
-	const int NumPoses = BlendPose.Num();
-	return FMath::Clamp<int32>(ActiveChildIndex, 0, NumPoses - 1);
+	const int32 NumPoses = BlendPose.Num();
+	int32 CurrentActiveChildIndex = GET_ANIM_NODE_DATA(int32, ActiveChildIndex);
+	return FMath::Clamp<int32>(CurrentActiveChildIndex, 0, NumPoses - 1);
 }
diff --git a/Engine/Source/Runtime/AnimGraphRuntime/Private/AnimNodes/AnimNode_BlendSpaceEvaluator.cpp b/Engine/Source/Runtime/AnimGraphRuntime/Private/AnimNodes/AnimNode_BlendSpaceEvaluator.cpp
--- a/Engine/Source/Runtime/AnimGraphRuntime/Private/AnimNodes/AnimNode_BlendSpaceEvaluator.cpp
+++ b/Engine/Source/Runtime/AnimGraphRuntime/Private/AnimNodes/AnimNode_BlendSpaceEvaluator.cpp
@@ -17,12 +17,11 @@ void FAnimNode_BlendSpaceEvaluator::UpdateAssetPlayer(const FAnimationUpdateCont
 {
 	GetEvaluateGraphExposedInputs().Execute(Context);
 	InternalTimeAccumulator = FMath::Clamp(NormalizedTime, 0.f, 1.f);
-	PlayRate = 0.f;
 
 	UpdateInternal(Context);
 
-	TRACE_ANIM_NODE_VALUE(Context, TEXT("Name"), BlendSpace ? *BlendSpace->GetName() : TEXT("None"));
-	TRACE_ANIM_NODE_VALUE(Context, TEXT("Blend Space"), BlendSpace);
+	TRACE_ANIM_NODE_VALUE(Context, TEXT("Name"), GetBlendSpace() ? *GetBlendSpace()->GetName() : TEXT("None"));
+	TRACE_ANIM_NODE_VALUE(Context, TEXT("Blend Space"), GetBlendSpace());
 	TRACE_ANIM_NODE_VALUE(Context, TEXT("Playback Time"), InternalTimeAccumulator);
 }
 
@@ -31,6 +30,6 @@ void FAnimNode_BlendSpaceEvaluator::GatherDebugData(FNodeDebugData& DebugData)
 	DECLARE_SCOPE_HIERARCHICAL_COUNTER_ANIMNODE(GatherDebugData)
 	FString DebugLine = DebugData.GetNodeName(this);
 
-	DebugLine += FString::Printf(TEXT("('%s' Play Time: %.3f)"), *BlendSpace->GetName(), InternalTimeAccumulator);
+	DebugLine += FString::Printf(TEXT("('%s' Play Time: %.3f)"), *GetBlendSpace()->GetName(), InternalTimeAccumulator);
 	DebugData.AddDebugItem(DebugLine, true);
 }
diff --git a/Engine/Source/Runtime/AnimGraphRuntime/Private/AnimNodes/AnimNode_BlendSpacePlayer.cpp b/Engine/Source/Runtime/AnimGraphRuntime/Private/AnimNodes/AnimNode_BlendSpacePlayer.cpp
--- a/Engine/Source/Runtime/AnimGraphRuntime/Private/AnimNodes/AnimNode_BlendSpacePlayer.cpp
+++ b/Engine/Source/Runtime/AnimGraphRuntime/Private/AnimNodes/AnimNode_BlendSpacePlayer.cpp
@@ -11,20 +11,7 @@
 /////////////////////////////////////////////////////
 // FAnimNode_BlendSpacePlayer
 
-FAnimNode_BlendSpacePlayer::FAnimNode_BlendSpacePlayer()
-	: X(0.0f)
-	, Y(0.0f)
-	, Z(0.0f)
-	, PlayRate(1.0f)
-	, bLoop(true)
-	, bResetPlayTimeWhenBlendSpaceChanges(true)
-	, StartPosition(0.f)
-	, BlendSpace(nullptr)
-	, PreviousBlendSpace(nullptr)
-{
-}
-
-float FAnimNode_BlendSpacePlayer::GetCurrentAssetTime()
+float FAnimNode_BlendSpacePlayer::GetCurrentAssetTime() const
 {
 	if(const FBlendSampleData* HighestWeightedSample = GetHighestWeightedSample())
 	{
@@ -35,19 +22,20 @@ float FAnimNode_BlendSpacePlayer::GetCurrentAssetTime()
 	return 0.0f;
 }
 
-float FAnimNode_BlendSpacePlayer::GetCurrentAssetTimePlayRateAdjusted()
+float FAnimNode_BlendSpacePlayer::GetCurrentAssetTimePlayRateAdjusted() const
 {
 	float Length = GetCurrentAssetLength();
-	return PlayRate < 0.0f ? Length - InternalTimeAccumulator * Length : Length * InternalTimeAccumulator;
+	return GetPlayRate() < 0.0f ? Length - InternalTimeAccumulator * Length : Length * InternalTimeAccumulator;
 }
 
-float FAnimNode_BlendSpacePlayer::GetCurrentAssetLength()
+float FAnimNode_BlendSpacePlayer::GetCurrentAssetLength() const
 {
 	if(const FBlendSampleData* HighestWeightedSample = GetHighestWeightedSample())
 	{
-		if (BlendSpace != nullptr)
+		UBlendSpace* CurrentBlendSpace = GetBlendSpace();
+		if (CurrentBlendSpace != nullptr)
 		{
-			const FBlendSample& Sample = BlendSpace->GetBlendSample(HighestWeightedSample->SampleDataIndex);
+			const FBlendSample& Sample = CurrentBlendSpace->GetBlendSample(HighestWeightedSample->SampleDataIndex);
 			return Sample.Animation->GetPlayLength();
 		}
 	}
@@ -65,7 +53,7 @@ void FAnimNode_BlendSpacePlayer::Initialize_AnyThread(const FAnimationInitialize
 
 	Reinitialize();
 
-	PreviousBlendSpace = BlendSpace;
+	PreviousBlendSpace = GetBlendSpace();
 }
 
 void FAnimNode_BlendSpacePlayer::CacheBones_AnyThread(const FAnimationCacheBonesContext& Context)
@@ -83,21 +71,23 @@ void FAnimNode_BlendSpacePlayer::UpdateAssetPlayer(const FAnimationUpdateContext
 void FAnimNode_BlendSpacePlayer::UpdateInternal(const FAnimationUpdateContext& Context)
 {
 	DECLARE_SCOPE_HIERARCHICAL_COUNTER_ANIMNODE(UpdateInternal)
-	if ((BlendSpace != nullptr) && (Context.AnimInstanceProxy->IsSkeletonCompatible(BlendSpace->GetSkeleton())))
+
+	UBlendSpace* CurrentBlendSpace = GetBlendSpace();
+	if ((CurrentBlendSpace != nullptr) && (Context.AnimInstanceProxy->IsSkeletonCompatible(CurrentBlendSpace->GetSkeleton())))
 	{	
-		if (PreviousBlendSpace != BlendSpace)
+		if (PreviousBlendSpace != CurrentBlendSpace)
 		{
-			Reinitialize(bResetPlayTimeWhenBlendSpaceChanges);
+			Reinitialize(ShouldResetPlayTimeWhenBlendSpaceChanges());
 		}
 
-		const FVector Position(X, Y, Z);
+		const FVector Position = GetPosition();
 
 		// Create a tick record and push into the closest scope
 		UE::Anim::FAnimSyncGroupScope& SyncScope = Context.GetMessageChecked<UE::Anim::FAnimSyncGroupScope>();
 
-		FAnimTickRecord TickRecord(BlendSpace, Position, BlendSampleDataCache, BlendFilter, bLoop, PlayRate, Context.GetFinalBlendWeight(), /*inout*/ InternalTimeAccumulator, MarkerTickRecord);
+		FAnimTickRecord TickRecord(CurrentBlendSpace, Position, BlendSampleDataCache, BlendFilter, GetLoop(), GetPlayRate(), Context.GetFinalBlendWeight(), /*inout*/ InternalTimeAccumulator, MarkerTickRecord);
 		TickRecord.RootMotionWeightModifier = Context.GetRootMotionWeightModifier();
-		UE::Anim::FAnimSyncParams SyncParams(GroupName, GroupRole, Method);
+		UE::Anim::FAnimSyncParams SyncParams(GetGroupName(), GetGroupRole(), GetGroupMethod());
 
 		SyncScope.AddTickRecord(TickRecord, SyncParams, UE::Anim::FAnimSyncDebugInfo(Context));
 
@@ -106,26 +96,28 @@ void FAnimNode_BlendSpacePlayer::UpdateInternal(const FAnimationUpdateContext& C
 #if WITH_EDITORONLY_DATA
 		if (FAnimBlueprintDebugData* DebugData = Context.AnimInstanceProxy->GetAnimBlueprintDebugData())
 		{
-			DebugData->RecordBlendSpacePlayer(Context.GetCurrentNodeId(), BlendSpace, Position, BlendFilter.GetFilterLastOutput());
+			DebugData->RecordBlendSpacePlayer(Context.GetCurrentNodeId(), CurrentBlendSpace, Position, BlendFilter.GetFilterLastOutput());
 		}
 #endif
 
-		PreviousBlendSpace = BlendSpace;
+		PreviousBlendSpace = CurrentBlendSpace;
 	}
 
 	TRACE_BLENDSPACE_PLAYER(Context, *this);
-	TRACE_ANIM_NODE_VALUE(Context, TEXT("Name"), BlendSpace ? *BlendSpace->GetName() : TEXT("None"));
-	TRACE_ANIM_NODE_VALUE(Context, TEXT("Blend Space"), BlendSpace);
+	TRACE_ANIM_NODE_VALUE(Context, TEXT("Name"), CurrentBlendSpace ? *CurrentBlendSpace->GetName() : TEXT("None"));
+	TRACE_ANIM_NODE_VALUE(Context, TEXT("Blend Space"), CurrentBlendSpace);
 	TRACE_ANIM_NODE_VALUE(Context, TEXT("Playback Time"), InternalTimeAccumulator);
 }
 
 void FAnimNode_BlendSpacePlayer::Evaluate_AnyThread(FPoseContext& Output)
 {
 	DECLARE_SCOPE_HIERARCHICAL_COUNTER_ANIMNODE(Evaluate_AnyThread)
-	if ((BlendSpace != nullptr) && (Output.AnimInstanceProxy->IsSkeletonCompatible(BlendSpace->GetSkeleton())))
+
+	UBlendSpace* CurrentBlendSpace = GetBlendSpace();
+	if ((CurrentBlendSpace != nullptr) && (Output.AnimInstanceProxy->IsSkeletonCompatible(CurrentBlendSpace->GetSkeleton())))
 	{
 		FAnimationPoseData AnimationPoseData(Output);
-		BlendSpace->GetAnimationPose(BlendSampleDataCache, AnimationPoseData);
+		CurrentBlendSpace->GetAnimationPose(BlendSampleDataCache, AnimationPoseData);
 	}
 	else
 	{
@@ -133,38 +125,30 @@ void FAnimNode_BlendSpacePlayer::Evaluate_AnyThread(FPoseContext& Output)
 	}
 }
 
-void FAnimNode_BlendSpacePlayer::OverrideAsset(UAnimationAsset* NewAsset)
-{
-	if(UBlendSpace* NewBlendSpace = Cast<UBlendSpace>(NewAsset))
-	{
-		BlendSpace = NewBlendSpace;
-	}
-}
-
 void FAnimNode_BlendSpacePlayer::GatherDebugData(FNodeDebugData& DebugData)
 {
 	DECLARE_SCOPE_HIERARCHICAL_COUNTER_ANIMNODE(GatherDebugData)
 	FString DebugLine = DebugData.GetNodeName(this);
-	if (BlendSpace)
+
+	UBlendSpace* CurrentBlendSpace = GetBlendSpace();
+	if (CurrentBlendSpace)
 	{
-		DebugLine += FString::Printf(TEXT("('%s' Play Time: %.3f)"), *BlendSpace->GetName(), InternalTimeAccumulator);
+		DebugLine += FString::Printf(TEXT("('%s' Play Time: %.3f)"), *CurrentBlendSpace->GetName(), InternalTimeAccumulator);
 
 		DebugData.AddDebugItem(DebugLine, true);
-
-		//BlendSpace->GetBlendSamples();
 	}
 }
 
-float FAnimNode_BlendSpacePlayer::GetTimeFromEnd(float CurrentTime)
+float FAnimNode_BlendSpacePlayer::GetTimeFromEnd(float CurrentTime) const
 {
 	// Blend-spaces use normalized time value
 	const float PlayLength = 1.0f;
-	return BlendSpace != nullptr ? PlayLength - CurrentTime : 0.0f;
+	return GetBlendSpace() != nullptr ? PlayLength - CurrentTime : 0.0f;
 }
 
-UAnimationAsset* FAnimNode_BlendSpacePlayer::GetAnimAsset()
+UAnimationAsset* FAnimNode_BlendSpacePlayer::GetAnimAsset() const
 {
-	return BlendSpace;
+	return GetBlendSpace();
 }
 
 const FBlendSampleData* FAnimNode_BlendSpacePlayer::GetHighestWeightedSample() const
@@ -192,16 +176,57 @@ void FAnimNode_BlendSpacePlayer::Reinitialize(bool bResetTime)
 	BlendSampleDataCache.Empty();
 	if(bResetTime)
 	{
-		InternalTimeAccumulator = FMath::Clamp(StartPosition, 0.f, 1.0f);
-		if (StartPosition == 0.f && PlayRate < 0.0f)
+		float CurrentStartPosition = GetStartPosition();
+
+		InternalTimeAccumulator = FMath::Clamp(CurrentStartPosition, 0.f, 1.0f);
+		if (CurrentStartPosition == 0.f && GetPlayRate() < 0.0f)
 		{
 			// Blend spaces run between 0 and 1
 			InternalTimeAccumulator = 1.0f;
 		}
 	}
 
-	if (BlendSpace != nullptr)
+	UBlendSpace* CurrentBlendSpace = GetBlendSpace();
+	if (CurrentBlendSpace != nullptr)
 	{
-		BlendSpace->InitializeFilter(&BlendFilter);
+		CurrentBlendSpace->InitializeFilter(&BlendFilter);
 	}
 }
+
+#if WITH_EDITORONLY_DATA
+void FAnimNode_BlendSpacePlayer::SetBlendSpace(UBlendSpace* InBlendSpace)
+{
+	TObjectPtr<UBlendSpace>& BlendSpaceToSet = GET_MUTABLE_ANIM_NODE_DATA(TObjectPtr<UBlendSpace>, BlendSpace);
+	BlendSpaceToSet = InBlendSpace;
+}
+#endif
+
+FVector FAnimNode_BlendSpacePlayer::GetPosition() const
+{
+	return FVector(GET_ANIM_NODE_DATA(float, X), GET_ANIM_NODE_DATA(float, Y), GET_ANIM_NODE_DATA(float, Z));
+}
+
+float FAnimNode_BlendSpacePlayer::GetPlayRate() const
+{
+	return GET_ANIM_NODE_DATA(float, PlayRate);
+}
+
+bool FAnimNode_BlendSpacePlayer::GetLoop() const
+{
+	return GET_ANIM_NODE_DATA(bool, bLoop);
+}
+
+bool FAnimNode_BlendSpacePlayer::ShouldResetPlayTimeWhenBlendSpaceChanges() const
+{
+	return GET_ANIM_NODE_DATA(bool, bResetPlayTimeWhenBlendSpaceChanges);
+}
+
+float FAnimNode_BlendSpacePlayer::GetStartPosition() const
+{
+	return GET_ANIM_NODE_DATA(float, StartPosition);
+}
+
+UBlendSpace* FAnimNode_BlendSpacePlayer::GetBlendSpace() const
+{
+	return GET_ANIM_NODE_DATA(TObjectPtr<UBlendSpace>, BlendSpace);
+}
diff --git a/Engine/Source/Runtime/AnimGraphRuntime/Private/AnimNodes/AnimNode_PoseHandler.cpp b/Engine/Source/Runtime/AnimGraphRuntime/Private/AnimNodes/AnimNode_PoseHandler.cpp
--- a/Engine/Source/Runtime/AnimGraphRuntime/Private/AnimNodes/AnimNode_PoseHandler.cpp
+++ b/Engine/Source/Runtime/AnimGraphRuntime/Private/AnimNodes/AnimNode_PoseHandler.cpp
@@ -83,13 +83,12 @@ void FAnimNode_PoseHandler::UpdateAssetPlayer(const FAnimationUpdateContext& Con
 	TRACE_ANIM_NODE_VALUE(Context, TEXT("Pose Asset"), CurrentPoseAsset.IsValid() ? *CurrentPoseAsset.Get()->GetName() : TEXT("None"));
 }
 
-void FAnimNode_PoseHandler::OverrideAsset(UAnimationAsset* NewAsset)
+#if WITH_EDITORONLY_DATA
+void FAnimNode_PoseHandler::SetPoseAsset(UPoseAsset* InPoseAsset)
 {
-	if(UPoseAsset* NewPoseAsset = Cast<UPoseAsset>(NewAsset))
-	{
-		PoseAsset = NewPoseAsset;
-	}
+	PoseAsset = InPoseAsset;
 }
+#endif
 
 void FAnimNode_PoseHandler::GatherDebugData(FNodeDebugData& DebugData)
 {
diff --git a/Engine/Source/Runtime/AnimGraphRuntime/Private/AnimNodes/AnimNode_RefPose.cpp b/Engine/Source/Runtime/AnimGraphRuntime/Private/AnimNodes/AnimNode_RefPose.cpp
--- a/Engine/Source/Runtime/AnimGraphRuntime/Private/AnimNodes/AnimNode_RefPose.cpp
+++ b/Engine/Source/Runtime/AnimGraphRuntime/Private/AnimNodes/AnimNode_RefPose.cpp
@@ -26,7 +26,8 @@ void FAnimNode_RefPose::Evaluate_AnyThread(FPoseContext& Output)
 	// I don't have anything to evaluate. Should this be even here?
 	// EvaluateGraphExposedInputs.Execute(Context);
 
-	switch (RefPoseType) 
+    ERefPoseType CurrentRefPoseType = GetRefPoseType(); 
+	switch (CurrentRefPoseType) 
 	{
 	case EIT_LocalSpace:
 		Output.ResetToRefPose();
@@ -38,7 +39,12 @@ void FAnimNode_RefPose::Evaluate_AnyThread(FPoseContext& Output)
 		break;
 	}
 
-	TRACE_ANIM_NODE_VALUE(Output, TEXT("Ref Pose Type"), GetRefPostTypeText(RefPoseType));
+	TRACE_ANIM_NODE_VALUE(Output, TEXT("Ref Pose Type"), GetRefPostTypeText(CurrentRefPoseType));
+}
+
+ERefPoseType FAnimNode_RefPose::GetRefPoseType() const
+{
+	return GET_ANIM_NODE_DATA(TEnumAsByte<ERefPoseType>, RefPoseType);
 }
 
 void FAnimNode_MeshSpaceRefPose::EvaluateComponentSpace_AnyThread(FComponentSpacePoseContext& Output)
@@ -51,6 +57,6 @@ void FAnimNode_RefPose::GatherDebugData(FNodeDebugData& DebugData)
 {
 	DECLARE_SCOPE_HIERARCHICAL_COUNTER_ANIMNODE(GatherDebugData)
 	FString DebugLine = DebugData.GetNodeName(this);
-	DebugLine += FString::Printf(TEXT("(Ref Pose Type: %s)"), GetRefPostTypeText(RefPoseType));
+	DebugLine += FString::Printf(TEXT("(Ref Pose Type: %s)"), GetRefPostTypeText(GetRefPoseType()));
 	DebugData.AddDebugItem(DebugLine, true);
 }
diff --git a/Engine/Source/Runtime/AnimGraphRuntime/Private/AnimNodes/AnimNode_SequenceEvaluator.cpp b/Engine/Source/Runtime/AnimGraphRuntime/Private/AnimNodes/AnimNode_SequenceEvaluator.cpp
--- a/Engine/Source/Runtime/AnimGraphRuntime/Private/AnimNodes/AnimNode_SequenceEvaluator.cpp
+++ b/Engine/Source/Runtime/AnimGraphRuntime/Private/AnimNodes/AnimNode_SequenceEvaluator.cpp
@@ -4,65 +4,67 @@
 #include "Animation/AnimInstanceProxy.h"
 #include "Animation/AnimTrace.h"
 
-float FAnimNode_SequenceEvaluator::GetCurrentAssetTime()
+float FAnimNode_SequenceEvaluatorBase::GetCurrentAssetTime() const
 {
-	return ExplicitTime;
+	return GetExplicitTime();
 }
 
-float FAnimNode_SequenceEvaluator::GetCurrentAssetLength()
+float FAnimNode_SequenceEvaluatorBase::GetCurrentAssetLength() const
 {
-	return Sequence ? Sequence->GetPlayLength() : 0.0f;
+	UAnimSequenceBase* CurrentSequence = GetSequence();
+	return CurrentSequence ? CurrentSequence->GetPlayLength() : 0.0f;
 }
 
-/////////////////////////////////////////////////////
-// FAnimSequenceEvaluatorNode
-
-void FAnimNode_SequenceEvaluator::Initialize_AnyThread(const FAnimationInitializeContext& Context)
+void FAnimNode_SequenceEvaluatorBase::Initialize_AnyThread(const FAnimationInitializeContext& Context)
 {
 	DECLARE_SCOPE_HIERARCHICAL_COUNTER_ANIMNODE(Initialize_AnyThread)
 	FAnimNode_AssetPlayerBase::Initialize_AnyThread(Context);
 	bReinitialized = true;
 }
 
-void FAnimNode_SequenceEvaluator::CacheBones_AnyThread(const FAnimationCacheBonesContext& Context)
+void FAnimNode_SequenceEvaluatorBase::CacheBones_AnyThread(const FAnimationCacheBonesContext& Context)
 {
 	DECLARE_SCOPE_HIERARCHICAL_COUNTER_ANIMNODE(CacheBones_AnyThread)
 }
 
-void FAnimNode_SequenceEvaluator::UpdateAssetPlayer(const FAnimationUpdateContext& Context)
+void FAnimNode_SequenceEvaluatorBase::UpdateAssetPlayer(const FAnimationUpdateContext& Context)
 {
 	GetEvaluateGraphExposedInputs().Execute(Context);
 
-	if (Sequence)
+	float CurrentExplicitTime = GetExplicitTime();
+
+	UAnimSequenceBase* CurrentSequence = GetSequence();
+	if (CurrentSequence)
 	{
 		// Clamp input to a valid position on this sequence's time line.
-		ExplicitTime = FMath::Clamp(ExplicitTime, 0.f, Sequence->GetPlayLength());
+		CurrentExplicitTime = FMath::Clamp(CurrentExplicitTime, 0.f, CurrentSequence->GetPlayLength());
 
-		if ((!bTeleportToExplicitTime || (GroupName != NAME_None)) && (Context.AnimInstanceProxy->IsSkeletonCompatible(Sequence->GetSkeleton())))
+		if ((!GetTeleportToExplicitTime() || (GetGroupName() != NAME_None)) && (Context.AnimInstanceProxy->IsSkeletonCompatible(CurrentSequence->GetSkeleton())))
 		{
 			if (bReinitialized)
 			{
-				switch (ReinitializationBehavior)
+				switch (GetReinitializationBehavior())
 				{
-					case ESequenceEvalReinit::StartPosition: InternalTimeAccumulator = StartPosition; break;
-					case ESequenceEvalReinit::ExplicitTime: InternalTimeAccumulator = ExplicitTime; break;
+					case ESequenceEvalReinit::StartPosition: InternalTimeAccumulator = GetStartPosition(); break;
+					case ESequenceEvalReinit::ExplicitTime: InternalTimeAccumulator = CurrentExplicitTime; break;
 				}
 
-				InternalTimeAccumulator = FMath::Clamp(InternalTimeAccumulator, 0.f, Sequence->GetPlayLength());
+				InternalTimeAccumulator = FMath::Clamp(InternalTimeAccumulator, 0.f, CurrentSequence->GetPlayLength());
 			}
 
-			float TimeJump = ExplicitTime - InternalTimeAccumulator;
-			if (bShouldLoop)
+			float TimeJump = CurrentExplicitTime - InternalTimeAccumulator;
+			const bool bCurrentShouldLoop = GetShouldLoop();
+			if (bCurrentShouldLoop)
 			{
-				if (FMath::Abs(TimeJump) > (Sequence->GetPlayLength() * 0.5f))
+				if (FMath::Abs(TimeJump) > (CurrentSequence->GetPlayLength() * 0.5f))
 				{
 					if (TimeJump > 0.f)
 					{
-						TimeJump -= Sequence->GetPlayLength();
+						TimeJump -= CurrentSequence->GetPlayLength();
 					}
 					else
 					{
-						TimeJump += Sequence->GetPlayLength();
+						TimeJump += CurrentSequence->GetPlayLength();
 					}
 				}
 			}
@@ -71,36 +73,37 @@ void FAnimNode_SequenceEvaluator::UpdateAssetPlayer(const FAnimationUpdateContex
 			// to prevent that from happening, we set current accumulator to explicit time
 			if (TimeJump == 0.f)
 			{
-				InternalTimeAccumulator = ExplicitTime;
+				InternalTimeAccumulator = CurrentExplicitTime;
 			}
 			
 			const float DeltaTime = Context.GetDeltaTime();
-			const float RateScale = Sequence->RateScale;
+			const float RateScale = CurrentSequence->RateScale;
 			const float PlayRate = FMath::IsNearlyZero(DeltaTime) || FMath::IsNearlyZero(RateScale) ? 0.f : (TimeJump / (DeltaTime * RateScale));
-			CreateTickRecordForNode(Context, Sequence, bShouldLoop, PlayRate);
+			CreateTickRecordForNode(Context, CurrentSequence, bCurrentShouldLoop, PlayRate);
 		}
 		else
 		{
-			InternalTimeAccumulator = ExplicitTime;
+			InternalTimeAccumulator = CurrentExplicitTime;
 		}
 	}
 
 	bReinitialized = false;
 
-	TRACE_ANIM_NODE_VALUE(Context, TEXT("Name"), Sequence != nullptr ? Sequence->GetFName() : NAME_None);
-	TRACE_ANIM_NODE_VALUE(Context, TEXT("Sequence"), Sequence);
-	TRACE_ANIM_NODE_VALUE(Context, TEXT("InputTime"), ExplicitTime);
+	TRACE_ANIM_NODE_VALUE(Context, TEXT("Name"), CurrentSequence != nullptr ? CurrentSequence->GetFName() : NAME_None);
+	TRACE_ANIM_NODE_VALUE(Context, TEXT("Sequence"), CurrentSequence);
+	TRACE_ANIM_NODE_VALUE(Context, TEXT("InputTime"), CurrentExplicitTime);
 	TRACE_ANIM_NODE_VALUE(Context, TEXT("Time"), InternalTimeAccumulator);
 }
 
-void FAnimNode_SequenceEvaluator::Evaluate_AnyThread(FPoseContext& Output)
+void FAnimNode_SequenceEvaluatorBase::Evaluate_AnyThread(FPoseContext& Output)
 {
 	DECLARE_SCOPE_HIERARCHICAL_COUNTER_ANIMNODE(Evaluate_AnyThread)
 	check(Output.AnimInstanceProxy != nullptr);
-	if ((Sequence != nullptr) && Output.AnimInstanceProxy->IsSkeletonCompatible(Sequence->GetSkeleton()))
+	UAnimSequenceBase* CurrentSequence = GetSequence();
+	if ((CurrentSequence != nullptr) && (Output.AnimInstanceProxy->IsSkeletonCompatible(CurrentSequence->GetSkeleton())))
 	{
 		FAnimationPoseData AnimationPoseData(Output);
-		Sequence->GetAnimationPose(AnimationPoseData, FAnimExtractContext(InternalTimeAccumulator, Output.AnimInstanceProxy->ShouldExtractRootMotion()));
+		CurrentSequence->GetAnimationPose(AnimationPoseData, FAnimExtractContext(InternalTimeAccumulator, Output.AnimInstanceProxy->ShouldExtractRootMotion()));
 	}
 	else
 	{
@@ -108,19 +111,49 @@ void FAnimNode_SequenceEvaluator::Evaluate_AnyThread(FPoseContext& Output)
 	}
 }
 
-void FAnimNode_SequenceEvaluator::OverrideAsset(UAnimationAsset* NewAsset)
-{
-	if(UAnimSequenceBase* NewSequence = Cast<UAnimSequenceBase>(NewAsset))
-	{
-		Sequence = NewSequence;
-	}
-}
-
-void FAnimNode_SequenceEvaluator::GatherDebugData(FNodeDebugData& DebugData)
+void FAnimNode_SequenceEvaluatorBase::GatherDebugData(FNodeDebugData& DebugData)
 {
 	DECLARE_SCOPE_HIERARCHICAL_COUNTER_ANIMNODE(GatherDebugData)
 	FString DebugLine = DebugData.GetNodeName(this);
 	
-	DebugLine += FString::Printf(TEXT("('%s' InputTime: %.3f, Time: %.3f)"), *GetNameSafe(Sequence), ExplicitTime, InternalTimeAccumulator);
+	DebugLine += FString::Printf(TEXT("('%s' InputTime: %.3f, Time: %.3f)"), *GetNameSafe(GetSequence()), GetExplicitTime(), InternalTimeAccumulator);
 	DebugData.AddDebugItem(DebugLine, true);
 }
+
+void FAnimNode_SequenceEvaluator::SetSequence(UAnimSequenceBase* InSequence)
+{
+#if WITH_EDITORONLY_DATA
+	TObjectPtr<UAnimSequenceBase>& SequenceToSet = GET_MUTABLE_ANIM_NODE_DATA(TObjectPtr<UAnimSequenceBase>, Sequence);
+	SequenceToSet = InSequence;
+#endif
+}
+
+UAnimSequenceBase* FAnimNode_SequenceEvaluator::GetSequence() const
+{
+	return GET_ANIM_NODE_DATA(TObjectPtr<UAnimSequenceBase>, Sequence);
+}
+
+float FAnimNode_SequenceEvaluator::GetExplicitTime() const
+{
+	return GET_ANIM_NODE_DATA(float, ExplicitTime);
+}
+
+bool FAnimNode_SequenceEvaluator::GetShouldLoop() const
+{
+	return GET_ANIM_NODE_DATA(bool, bShouldLoop);
+}
+
+bool FAnimNode_SequenceEvaluator::GetTeleportToExplicitTime() const
+{
+	return GET_ANIM_NODE_DATA(bool, bTeleportToExplicitTime);
+}
+
+TEnumAsByte<ESequenceEvalReinit::Type> FAnimNode_SequenceEvaluator::GetReinitializationBehavior() const
+{
+	return GET_ANIM_NODE_DATA(TEnumAsByte<ESequenceEvalReinit::Type>, ReinitializationBehavior);
+}
+
+float FAnimNode_SequenceEvaluator::GetStartPosition() const
+{
+	return GET_ANIM_NODE_DATA(float, StartPosition);
+}
diff --git a/Engine/Source/Runtime/AnimGraphRuntime/Private/AnimSequencerInstanceProxy.cpp b/Engine/Source/Runtime/AnimGraphRuntime/Private/AnimSequencerInstanceProxy.cpp
--- a/Engine/Source/Runtime/AnimGraphRuntime/Private/AnimSequencerInstanceProxy.cpp
+++ b/Engine/Source/Runtime/AnimGraphRuntime/Private/AnimSequencerInstanceProxy.cpp
@@ -132,8 +132,7 @@ void FAnimSequencerInstanceProxy::InitAnimTrack(UAnimSequenceBase* InAnimSequenc
 			SequencerToPlayerMap.Add(SequenceId, NewPlayerState);
 
 			// link player to blendnode, this will let you trigger notifies and so on
-			NewPlayerState->PlayerNode.bTeleportToExplicitTime = false;
-			NewPlayerState->PlayerNode.bShouldLoop = true;
+			NewPlayerState->PlayerNode.SetTeleportToExplicitTime(false);
 			BlendNode.Poses[PoseIndex].SetLinkNode(&NewPlayerState->PlayerNode);
 
 			// set player state
@@ -141,8 +140,8 @@ void FAnimSequencerInstanceProxy::InitAnimTrack(UAnimSequenceBase* InAnimSequenc
 		}
 
 		// now set animation data to player
-		PlayerState->PlayerNode.Sequence = InAnimSequence;
-		PlayerState->PlayerNode.ExplicitTime = 0.f;
+		PlayerState->PlayerNode.SetSequence(InAnimSequence);
+		PlayerState->PlayerNode.SetExplicitTime(0.f);
 
 		// initialize player
 		PlayerState->PlayerNode.Initialize_AnyThread(FAnimationInitializeContext(this));
@@ -194,17 +193,17 @@ void FAnimSequencerInstanceProxy::UpdateAnimTrack(UAnimSequenceBase* InAnimSeque
 
 	FSequencerPlayerAnimSequence* PlayerState = FindPlayer<FSequencerPlayerAnimSequence>(SequenceId);
 
-	PlayerState->PlayerNode.ExplicitTime = InToPosition;
+	PlayerState->PlayerNode.SetExplicitTime(InToPosition);
 	if (InFromPosition.IsSet())
 	{
 		// Set the internal time accumulator at the "from" time so that the player node will correctly evaluate the
 		// desired "from/to" range. We also disable the reinitialization code so it doesn't mess up that time we
 		// just set.
 		PlayerState->PlayerNode.SetExplicitPreviousTime(InFromPosition.GetValue());
-		PlayerState->PlayerNode.ReinitializationBehavior = ESequenceEvalReinit::NoReset;
+		PlayerState->PlayerNode.SetReinitializationBehavior(ESequenceEvalReinit::NoReset);
 	}
 	// if no fire notifies, we can teleport to explicit time
-	PlayerState->PlayerNode.bTeleportToExplicitTime = !bFireNotifies;
+	PlayerState->PlayerNode.SetTeleportToExplicitTime(!bFireNotifies);
 	// if moving to 0.f, we mark this to teleport. Otherwise, do not use explicit time
 	FAnimNode_MultiWayBlend& BlendNode = (PlayerState->bAdditive) ? AdditiveBlendNode : FullBodyBlendNode;
 	BlendNode.DesiredAlphas[PlayerState->PoseIndex] = Weight;
@@ -224,9 +223,9 @@ void FAnimSequencerInstanceProxy::EnsureAnimTrack(UAnimSequenceBase* InAnimSeque
 	{
 		InitAnimTrack(InAnimSequence, SequenceId);
 	}
-	else if (PlayerState->PlayerNode.Sequence != InAnimSequence)
+	else if (PlayerState->PlayerNode.GetSequence() != InAnimSequence)
 	{
-		PlayerState->PlayerNode.OverrideAsset(InAnimSequence);
+		PlayerState->PlayerNode.SetSequence(InAnimSequence);
 	}
 }
 
diff --git a/Engine/Source/Runtime/AnimGraphRuntime/Public/AnimNodes/AnimNode_AimOffsetLookAt.h b/Engine/Source/Runtime/AnimGraphRuntime/Public/AnimNodes/AnimNode_AimOffsetLookAt.h
--- a/Engine/Source/Runtime/AnimGraphRuntime/Public/AnimNodes/AnimNode_AimOffsetLookAt.h
+++ b/Engine/Source/Runtime/AnimGraphRuntime/Public/AnimNodes/AnimNode_AimOffsetLookAt.h
@@ -55,6 +55,9 @@ struct ANIMGRAPHRUNTIME_API FAnimNode_AimOffsetLookAt : public FAnimNode_BlendSp
 	UPROPERTY(EditAnywhere, BlueprintReadWrite, Category = LookAt, meta = (PinShownByDefault))
 	float Alpha;
 
+	/** Cached calculated blend input */
+	FVector CurrentBlendInput;
+
 	/** Cached reference to the source socket's bone */
 	FBoneReference SocketBoneReference;
 
@@ -78,5 +81,8 @@ public:
 	virtual int32 GetLODThreshold() const override { return LODThreshold; }
 	// End of FAnimNode_Base interface
 
+	// FAnimNode_BlendSpacePlayer interface
+	virtual FVector GetPosition() const override;
+
 	void UpdateFromLookAtTarget(FPoseContext& LocalPoseContext);
 };
diff --git a/Engine/Source/Runtime/AnimGraphRuntime/Public/AnimNodes/AnimNode_BlendListBase.h b/Engine/Source/Runtime/AnimGraphRuntime/Public/AnimNodes/AnimNode_BlendListBase.h
--- a/Engine/Source/Runtime/AnimGraphRuntime/Public/AnimNodes/AnimNode_BlendListBase.h
+++ b/Engine/Source/Runtime/AnimGraphRuntime/Public/AnimNodes/AnimNode_BlendListBase.h
@@ -23,34 +23,36 @@ enum class EBlendListTransitionType : uint8
 USTRUCT(BlueprintInternalUseOnly)
 struct ANIMGRAPHRUNTIME_API FAnimNode_BlendListBase : public FAnimNode_Base
 {
-	GENERATED_USTRUCT_BODY()
-public:
-	UPROPERTY(EditAnywhere, EditFixedSize, BlueprintReadWrite, Category=Links)
+	GENERATED_BODY()
+
+protected:	
+	UPROPERTY(EditAnywhere, EditFixedSize, Category=Links)
 	TArray<FPoseLink> BlendPose;
 
-	UPROPERTY(EditAnywhere, EditFixedSize, BlueprintReadWrite, Category=Config, meta=(PinShownByDefault))
+private:
+#if WITH_EDITORONLY_DATA
+	UPROPERTY(EditAnywhere, EditFixedSize, Category=Config, meta=(PinShownByDefault, FoldProperty))
 	TArray<float> BlendTime;
 
-	UPROPERTY(EditAnywhere, BlueprintReadWrite, Category=Config)
-	EBlendListTransitionType TransitionType;
+	UPROPERTY(EditAnywhere, Category=Config, meta=(FoldProperty))
+	EBlendListTransitionType TransitionType = EBlendListTransitionType::StandardBlend;
 
-	UPROPERTY(EditAnywhere, Category=BlendType)
-	EAlphaBlendOption BlendType;
+	UPROPERTY(EditAnywhere, Category=BlendType, meta=(FoldProperty))
+	EAlphaBlendOption BlendType = EAlphaBlendOption::Linear;
 	
 protected:
 	/** This reinitializes child pose when re-activated. For example, when active child changes */
-	UPROPERTY(EditAnywhere, Category = Option)
-	bool bResetChildOnActivation;
-
-	int32 LastActiveChildIndex;
-
-public:
-	UPROPERTY(EditAnywhere, Category=BlendType)
-	TObjectPtr<UCurveFloat> CustomBlendCurve;
+	UPROPERTY(EditAnywhere, Category = Option, meta=(FoldProperty))
+	bool bResetChildOnActivation = false;
 
-	UPROPERTY(EditAnywhere, Category=BlendType, meta=(UseAsBlendProfile=true))
-	TObjectPtr<UBlendProfile> BlendProfile;
+private:
+	UPROPERTY(EditAnywhere, Category=BlendType, meta=(FoldProperty))
+	TObjectPtr<UCurveFloat> CustomBlendCurve = nullptr;
 
+	UPROPERTY(EditAnywhere, Category=BlendType, meta=(UseAsBlendProfile=true, FoldProperty))
+	TObjectPtr<UBlendProfile> BlendProfile = nullptr;
+#endif // #if WITH_EDITORONLY_DATA
+	
 	TArray<struct FAlphaBlend> Blends;
 
 protected:
@@ -65,16 +67,10 @@ protected:
 	//Store which poses we need to evaluate
 	TArray<int32> PosesToEvaluate;
 
+	int32 LastActiveChildIndex = 0;
+	
 public:	
-	FAnimNode_BlendListBase()
-		: TransitionType(EBlendListTransitionType::StandardBlend)
-		, BlendType(EAlphaBlendOption::Linear)
-		, bResetChildOnActivation(false)
-		, LastActiveChildIndex(0)
-		, CustomBlendCurve(nullptr)
-		, BlendProfile(nullptr)
-	{
-	}
+	FAnimNode_BlendListBase() = default;
 
 	// FAnimNode_Base interface
 	virtual void Initialize_AnyThread(const FAnimationInitializeContext& Context) override;
@@ -98,6 +94,25 @@ public:
 	}
 #endif
 
+public:
+	// Get the array of blend times to apply to our input poses
+	const TArray<float>& GetBlendTimes() const;
+
+	// Get the type of transition that this blend list will make
+	EBlendListTransitionType GetTransitionType() const;
+
+	// Get the blend type we will use when blending
+	EAlphaBlendOption GetBlendType() const;
+	
+	/** Get whether to reinitialize the child pose when re-activated. For example, when active child changes */
+	bool GetResetChildOnActivation() const;
+
+	// Get the custom blend curve to apply when blending, if any
+	UCurveFloat* GetCustomBlendCurve() const;
+
+	// Get the blend profile to use when blending, if any
+	UBlendProfile* GetBlendProfile() const;
+	
 protected:
 	virtual int32 GetActiveChildIndex() { return 0; }
 	virtual FString GetNodeName(FNodeDebugData& DebugData) { return DebugData.GetNodeName(this); }
diff --git a/Engine/Source/Runtime/AnimGraphRuntime/Public/AnimNodes/AnimNode_BlendListByBool.h b/Engine/Source/Runtime/AnimGraphRuntime/Public/AnimNodes/AnimNode_BlendListByBool.h
--- a/Engine/Source/Runtime/AnimGraphRuntime/Public/AnimNodes/AnimNode_BlendListByBool.h
+++ b/Engine/Source/Runtime/AnimGraphRuntime/Public/AnimNodes/AnimNode_BlendListByBool.h
@@ -11,19 +11,20 @@
 USTRUCT(BlueprintInternalUseOnly)
 struct ANIMGRAPHRUNTIME_API FAnimNode_BlendListByBool : public FAnimNode_BlendListBase
 {
-	GENERATED_USTRUCT_BODY()
-public:
+	GENERATED_BODY()
+private:
+#if WITH_EDITORONLY_DATA
 	// Which input should be connected to the output?
-	UPROPERTY(EditAnywhere, BlueprintReadWrite, Category=Runtime, meta=(PinShownByDefault))
-	bool bActiveValue;
-
+	UPROPERTY(EditAnywhere, Category=Runtime, meta=(PinShownByDefault, FoldProperty))
+	bool bActiveValue = false;
+#endif
+	
 public:	
-	FAnimNode_BlendListByBool()
-		: FAnimNode_BlendListBase()
-		, bActiveValue(false)
-	{
-	}
+	FAnimNode_BlendListByBool() = default;
 
+	// Get which input should be connected to the output
+	bool GetActiveValue() const;
+	
 protected:
 	virtual int32 GetActiveChildIndex() override;
 	virtual FString GetNodeName(FNodeDebugData& DebugData) override { return DebugData.GetNodeName(this); }
diff --git a/Engine/Source/Runtime/AnimGraphRuntime/Public/AnimNodes/AnimNode_BlendListByEnum.h b/Engine/Source/Runtime/AnimGraphRuntime/Public/AnimNodes/AnimNode_BlendListByEnum.h
--- a/Engine/Source/Runtime/AnimGraphRuntime/Public/AnimNodes/AnimNode_BlendListByEnum.h
+++ b/Engine/Source/Runtime/AnimGraphRuntime/Public/AnimNodes/AnimNode_BlendListByEnum.h
@@ -11,23 +11,33 @@
 USTRUCT(BlueprintInternalUseOnly)
 struct ANIMGRAPHRUNTIME_API FAnimNode_BlendListByEnum : public FAnimNode_BlendListBase
 {
-	GENERATED_USTRUCT_BODY()
-public:
+	GENERATED_BODY()
+	
+private:
+#if WITH_EDITORONLY_DATA
 	// Mapping from enum value to BlendPose index; there will be one entry per entry in the enum; entries out of range always map to pose index 0
-	UPROPERTY()
+	UPROPERTY(meta=(FoldProperty))
 	TArray<int32> EnumToPoseIndex;
-
+	
 	// The currently selected pose (as an enum value)
-	UPROPERTY(EditAnywhere, BlueprintReadWrite, Category=Runtime, meta=(PinShownByDefault))
-	mutable uint8 ActiveEnumValue;
-
+	UPROPERTY(EditAnywhere, Category=Runtime, meta=(PinShownByDefault, FoldProperty))
+	mutable uint8 ActiveEnumValue = 0;
+#endif
+	
 public:	
-	FAnimNode_BlendListByEnum()
-		: FAnimNode_BlendListBase()
-		, ActiveEnumValue(0)
-	{
-	}
+	FAnimNode_BlendListByEnum() = default;
 
+#if WITH_EDITORONLY_DATA
+	// Set the mapping from enum value to BlendPose index. Called during compilation.
+	void SetEnumToPoseIndex(const TArray<int32>& InEnumToPoseIndex);
+#endif
+	
+	// Get the mapping from enum value to BlendPose index; there will be one entry per entry in the enum; entries out of range always map to pose index 0
+	const TArray<int32>& GetEnumToPoseIndex() const;
+	
+	// Get the currently selected pose (as an enum value)
+	uint8 GetActiveEnumValue() const;
+	
 protected:
 	virtual int32 GetActiveChildIndex() override;
 	virtual FString GetNodeName(FNodeDebugData& DebugData) override { return DebugData.GetNodeName(this); }
diff --git a/Engine/Source/Runtime/AnimGraphRuntime/Public/AnimNodes/AnimNode_BlendListByInt.h b/Engine/Source/Runtime/AnimGraphRuntime/Public/AnimNodes/AnimNode_BlendListByInt.h
--- a/Engine/Source/Runtime/AnimGraphRuntime/Public/AnimNodes/AnimNode_BlendListByInt.h
+++ b/Engine/Source/Runtime/AnimGraphRuntime/Public/AnimNodes/AnimNode_BlendListByInt.h
@@ -11,19 +11,19 @@
 USTRUCT(BlueprintInternalUseOnly)
 struct ANIMGRAPHRUNTIME_API FAnimNode_BlendListByInt : public FAnimNode_BlendListBase
 {
-	GENERATED_USTRUCT_BODY()
+	GENERATED_BODY()
+private:
+#if WITH_EDITORONLY_DATA
+	UPROPERTY(EditAnywhere,  Category=Runtime, meta=(PinShownByDefault, FoldProperty))
+	int32 ActiveChildIndex = 0;
+#endif
+	
 public:
-	UPROPERTY(EditAnywhere, BlueprintReadWrite, Category=Runtime, meta=(PinShownByDefault))
-	int32 ActiveChildIndex;
-
-public:	
-	FAnimNode_BlendListByInt()
-		:	FAnimNode_BlendListBase()
-		,	ActiveChildIndex(0)
-	{
-	}
+	FAnimNode_BlendListByInt() = default;
 
+	// Get the currently active child index
+	virtual int32 GetActiveChildIndex() override;
+	
 protected:
-	virtual int32 GetActiveChildIndex();
-	virtual FString GetNodeName(FNodeDebugData& DebugData) { return DebugData.GetNodeName(this); }
+	virtual FString GetNodeName(FNodeDebugData& DebugData) override { return DebugData.GetNodeName(this); }
 };
diff --git a/Engine/Source/Runtime/AnimGraphRuntime/Public/AnimNodes/AnimNode_BlendSpacePlayer.h b/Engine/Source/Runtime/AnimGraphRuntime/Public/AnimNodes/AnimNode_BlendSpacePlayer.h
--- a/Engine/Source/Runtime/AnimGraphRuntime/Public/AnimNodes/AnimNode_BlendSpacePlayer.h
+++ b/Engine/Source/Runtime/AnimGraphRuntime/Public/AnimNodes/AnimNode_BlendSpacePlayer.h
@@ -14,67 +14,73 @@ class UBlendSpace;
 USTRUCT(BlueprintInternalUseOnly)
 struct ANIMGRAPHRUNTIME_API FAnimNode_BlendSpacePlayer : public FAnimNode_AssetPlayerBase
 {
-	GENERATED_USTRUCT_BODY()
+	GENERATED_BODY()
 
-	// @return the current sample coordinates that this node is using to sample the blendspace
-	FVector GetPosition() const { return FVector(X, Y, Z); }
+ 	friend class UAnimGraphNode_BlendSpacePlayer;
+ 	friend class UAnimGraphNode_BlendSpaceEvaluator;
+	friend class UAnimGraphNode_RotationOffsetBlendSpace;
+	friend class UAnimGraphNode_AimOffsetLookAt;
 
 	// @return the current sample coordinates after going through the filtering
 	FVector GetFilteredPosition() const { return BlendFilter.GetFilterLastOutput(); }
+private:
 
-public:
+#if WITH_EDITORONLY_DATA
 	// The X coordinate to sample in the blendspace
-	UPROPERTY(EditAnywhere, BlueprintReadWrite, Category=Coordinates, meta=(PinShownByDefault))
-	float X;
+	UPROPERTY(EditAnywhere, Category=Coordinates, meta=(PinShownByDefault, FoldProperty))
+	float X = 0.0f;
 
 	// The Y coordinate to sample in the blendspace
-	UPROPERTY(EditAnywhere, BlueprintReadWrite, Category = Coordinates, meta = (PinShownByDefault))
-	float Y;
+	UPROPERTY(EditAnywhere, Category = Coordinates, meta = (PinShownByDefault, FoldProperty))
+	float Y = 0.0f;
 
 	// The Z coordinate to sample in the blendspace
-	UPROPERTY(EditAnywhere, BlueprintReadWrite, Category = Coordinates, meta = (PinHiddenByDefault))
-	float Z;
+	UPROPERTY(EditAnywhere, Category = Coordinates, meta = (PinHiddenByDefault, FoldProperty))
+	float Z = 0.0f;
 
 	// The play rate multiplier. Can be negative, which will cause the animation to play in reverse.
-	UPROPERTY(EditAnywhere, BlueprintReadWrite, Category = Settings, meta = (DefaultValue = "1.0", PinHiddenByDefault))
-	float PlayRate;
+	UPROPERTY(EditAnywhere, Category = Settings, meta = (DefaultValue = "1.0", PinHiddenByDefault, FoldProperty))
+	float PlayRate = 1.0f;
 
 	// Should the animation continue looping when it reaches the end?
-	UPROPERTY(EditAnywhere, BlueprintReadWrite, Category = Settings, meta = (DefaultValue = "true", PinHiddenByDefault))
-	bool bLoop;
+	UPROPERTY(EditAnywhere, Category = Settings, meta = (DefaultValue = "true", PinHiddenByDefault, FoldProperty))
+	bool bLoop = true;
 
 	// Whether we should reset the current play time when the blend space changes
-	UPROPERTY(EditAnywhere, BlueprintReadWrite, Category = Settings, meta = (PinHiddenByDefault))
-	bool bResetPlayTimeWhenBlendSpaceChanges;
+	UPROPERTY(EditAnywhere, Category = Settings, meta = (PinHiddenByDefault, FoldProperty))
+	bool bResetPlayTimeWhenBlendSpaceChanges = true;
 
 	// The start up position in [0, 1], it only applies when reinitialized
 	// if you loop, it will still start from 0.f after finishing the round
-	UPROPERTY(EditAnywhere, BlueprintReadWrite, Category=Settings, meta = (DefaultValue = "0.f", PinHiddenByDefault))
-	float StartPosition;
+	UPROPERTY(EditAnywhere, Category=Settings, meta = (DefaultValue = "0.f", PinHiddenByDefault, FoldProperty))
+	float StartPosition = 0.0f;
 
 	// The blendspace asset to play
-	UPROPERTY(EditAnywhere, BlueprintReadWrite, Category = Settings, meta = (PinHiddenByDefault))
-	TObjectPtr<UBlendSpace> BlendSpace;
+	UPROPERTY(EditAnywhere, Category = Settings, meta = (PinHiddenByDefault, FoldProperty))
+	TObjectPtr<UBlendSpace> BlendSpace = nullptr;
+#endif
 
 protected:
-
+	// Filter used to dampen coordinate changes
 	FBlendFilter BlendFilter;
 
+	// Cache of samples used to determine blend weights
 	TArray<FBlendSampleData> BlendSampleDataCache;
 
 	/** Previous position in the triangulation/segmentation */
 	int32 CachedTriangulationIndex = -1;
 
 	UPROPERTY(Transient)
-	TObjectPtr<UBlendSpace> PreviousBlendSpace;
+	TObjectPtr<UBlendSpace> PreviousBlendSpace = nullptr;
 
 public:	
-	FAnimNode_BlendSpacePlayer();
+	FAnimNode_BlendSpacePlayer() = default;
 
 	// FAnimNode_AssetPlayerBase interface
-	virtual float GetCurrentAssetTime();
-	virtual float GetCurrentAssetTimePlayRateAdjusted();
-	virtual float GetCurrentAssetLength();
+	virtual float GetCurrentAssetTime() const override;
+	virtual float GetCurrentAssetTimePlayRateAdjusted() const override;
+	virtual float GetCurrentAssetLength() const override;
+	virtual UAnimationAsset* GetAnimAsset() const override;
 	// End of FAnimNode_AssetPlayerBase interface
 
 	// FAnimNode_Base interface
@@ -82,13 +88,34 @@ public:
 	virtual void CacheBones_AnyThread(const FAnimationCacheBonesContext& Context) override;
 	virtual void UpdateAssetPlayer(const FAnimationUpdateContext& Context) override;
 	virtual void Evaluate_AnyThread(FPoseContext& Output) override;
-	virtual void OverrideAsset(UAnimationAsset* NewAsset) override;
 	virtual void GatherDebugData(FNodeDebugData& DebugData) override;
 	// End of FAnimNode_Base interface
 
-	float GetTimeFromEnd(float CurrentTime);
+	float GetTimeFromEnd(float CurrentTime) const;
+
+#if WITH_EDITORONLY_DATA
+	// Set the blendspace asset to play
+	void SetBlendSpace(UBlendSpace* InBlendSpace);
+#endif
+
+	// Get the coordinates that are currently being sampled by the blendspace
+	virtual FVector GetPosition() const;
+
+	// Get the play rate multiplier. Can be negative, which will cause the animation to play in reverse.
+	float GetPlayRate() const;
+
+	// Should the animation continue looping when it reaches the end?
+	bool GetLoop() const;
+
+	// Get whether we should reset the current play time when the blend space changes
+	bool ShouldResetPlayTimeWhenBlendSpaceChanges() const;
+
+	// Get the start up position in [0, 1], it only applies when reinitialized
+	// if you loop, it will still start from 0.f after finishing the round
+	float GetStartPosition() const;
 
-	UAnimationAsset* GetAnimAsset();
+	// Get the blendspace asset to play
+	UBlendSpace* GetBlendSpace() const;
 
 protected:
 	void UpdateInternal(const FAnimationUpdateContext& Context);
diff --git a/Engine/Source/Runtime/AnimGraphRuntime/Public/AnimNodes/AnimNode_PoseHandler.h b/Engine/Source/Runtime/AnimGraphRuntime/Public/AnimNodes/AnimNode_PoseHandler.h
--- a/Engine/Source/Runtime/AnimGraphRuntime/Public/AnimNodes/AnimNode_PoseHandler.h
+++ b/Engine/Source/Runtime/AnimGraphRuntime/Public/AnimNodes/AnimNode_PoseHandler.h
@@ -28,24 +28,28 @@ public:
 	}
 
 	// FAnimNode_AssetPlayerBase interface
-	virtual float GetCurrentAssetTime() { return 0.f; }
-	virtual float GetCurrentAssetLength() { return 0.f; }
+	virtual float GetCurrentAssetTime() const { return 0.f; }
+	virtual float GetCurrentAssetLength() const { return 0.f; }
 	// End of FAnimNode_AssetPlayerBase interface
 
 	// FAnimNode_Base interface
 	virtual void Initialize_AnyThread(const FAnimationInitializeContext& Context) override;
 	virtual void CacheBones_AnyThread(const FAnimationCacheBonesContext& Context) override;
 	virtual void UpdateAssetPlayer(const FAnimationUpdateContext& Context) override;
-	virtual void OverrideAsset(UAnimationAsset* NewAsset) override;
 	virtual void GatherDebugData(FNodeDebugData& DebugData) override;
 	// End of FAnimNode_Base interface
 
 	// FAnimNode_AssetPlayerBase Interface
 	virtual float GetAccumulatedTime() const {return 0.f;}
 	virtual void SetAccumulatedTime(const float& NewTime) {}
-	virtual UAnimationAsset* GetAnimAsset() {return PoseAsset;}
+	virtual UAnimationAsset* GetAnimAsset() const {return PoseAsset;}
 	// End of FAnimNode_AssetPlayerBase Interface
 
+#if WITH_EDITORONLY_DATA
+	// Set the pose asset to use for this node 
+	void SetPoseAsset(UPoseAsset* InPoseAsset);
+#endif
+	
 protected:
 	/** Called after CurrentPoseAsset is changed.  */
 	virtual void OnPoseAssetChange() {}
diff --git a/Engine/Source/Runtime/AnimGraphRuntime/Public/AnimNodes/AnimNode_RefPose.h b/Engine/Source/Runtime/AnimGraphRuntime/Public/AnimNodes/AnimNode_RefPose.h
--- a/Engine/Source/Runtime/AnimGraphRuntime/Public/AnimNodes/AnimNode_RefPose.h
+++ b/Engine/Source/Runtime/AnimGraphRuntime/Public/AnimNodes/AnimNode_RefPose.h
@@ -18,32 +18,35 @@ enum ERefPoseType
 USTRUCT()
 struct ANIMGRAPHRUNTIME_API FAnimNode_RefPose : public FAnimNode_Base
 {
-	GENERATED_USTRUCT_BODY()
+	GENERATED_BODY()
 
-	UPROPERTY()
-	TEnumAsByte<ERefPoseType> RefPoseType;
+private:
+#if WITH_EDITORONLY_DATA
+	UPROPERTY(meta=(FoldProperty))
+	TEnumAsByte<ERefPoseType> RefPoseType = EIT_LocalSpace;
+#endif	// #if WITH_EDITORONLY_DATA
 	
 public:	
-	FAnimNode_RefPose()
-		:	FAnimNode_Base()
-		, 	RefPoseType(EIT_LocalSpace)
-	{
-	}
+	FAnimNode_RefPose() = default;
 
+#if WITH_EDITORONLY_DATA
+	// Set the ref pose type of this node
+	void SetRefPoseType(ERefPoseType InType) { RefPoseType = InType; }
+#endif
+
+	// Get the type of this ref pose
+	ERefPoseType GetRefPoseType() const;
+	
 	virtual void Evaluate_AnyThread(FPoseContext& Output) override;
 	virtual void GatherDebugData(FNodeDebugData& DebugData) override;
-
 };
 
 USTRUCT()
 struct ANIMGRAPHRUNTIME_API FAnimNode_MeshSpaceRefPose : public FAnimNode_Base
 {
-	GENERATED_USTRUCT_BODY()
+	GENERATED_BODY()
 public:	
-	FAnimNode_MeshSpaceRefPose()
-		:	FAnimNode_Base()
-	{
-	}
+	FAnimNode_MeshSpaceRefPose() = default;
 
 	virtual void EvaluateComponentSpace_AnyThread(FComponentSpacePoseContext& Output);
 };
diff --git a/Engine/Source/Runtime/AnimGraphRuntime/Public/AnimNodes/AnimNode_SequenceEvaluator.h b/Engine/Source/Runtime/AnimGraphRuntime/Public/AnimNodes/AnimNode_SequenceEvaluator.h
--- a/Engine/Source/Runtime/AnimGraphRuntime/Public/AnimNodes/AnimNode_SequenceEvaluator.h
+++ b/Engine/Source/Runtime/AnimGraphRuntime/Public/AnimNodes/AnimNode_SequenceEvaluator.h
@@ -22,73 +22,163 @@ namespace ESequenceEvalReinit
 	};
 }
 
-// Evaluates a point in an anim sequence, using a specific time input rather than advancing time internally.
+// Abstract base class. Evaluates a point in an anim sequence, using a specific time input rather than advancing time internally.
 // Typically the playback position of the animation for this node will represent something other than time, like jump height.
 // This node will not trigger any notifies present in the associated sequence.
 USTRUCT(BlueprintInternalUseOnly)
-struct ANIMGRAPHRUNTIME_API FAnimNode_SequenceEvaluator : public FAnimNode_AssetPlayerBase
+struct ANIMGRAPHRUNTIME_API FAnimNode_SequenceEvaluatorBase : public FAnimNode_AssetPlayerBase
 {
-	GENERATED_USTRUCT_BODY()
+	GENERATED_BODY()
+
+private:
+	bool bReinitialized = false;
+
 public:
+	// FAnimNode_AssetPlayerBase interface
+	virtual float GetCurrentAssetTime() const override;
+	virtual float GetCurrentAssetLength() const override;
+	// End of FAnimNode_AssetPlayerBase interface
+
+	// FAnimNode_Base interface
+	virtual void Initialize_AnyThread(const FAnimationInitializeContext& Context) override;
+	virtual void CacheBones_AnyThread(const FAnimationCacheBonesContext& Context) override;
+	virtual void UpdateAssetPlayer(const FAnimationUpdateContext& Context) override;
+	virtual void Evaluate_AnyThread(FPoseContext& Output) override;
+	virtual void GatherDebugData(FNodeDebugData& DebugData) override;
+	// End of FAnimNode_Base interface
+
+	// FAnimNode_AssetPlayerBase Interface
+	virtual float GetAccumulatedTime() const {return GetExplicitTime();}
+	virtual void SetAccumulatedTime(const float& NewTime) { SetExplicitTime(NewTime); }
+	virtual UAnimationAsset* GetAnimAsset() const { return GetSequence(); }
+	// End of FAnimNode_AssetPlayerBase Interface
+
+	void SetExplicitPreviousTime(float PreviousTime) { InternalTimeAccumulator = PreviousTime; }
+
+	// Set the animation sequence asset to evaluate
+	virtual void SetSequence(UAnimSequenceBase* InSequence) PURE_VIRTUAL(SetSequence, )
+
+	// Set the time at which to evaluate the associated sequence
+	virtual void SetExplicitTime(float InTime) PURE_VIRTUAL(SetExplicitTime, )
+
+	// Set whether to teleport to explicit time when it is set
+	virtual void SetTeleportToExplicitTime(bool bInTeleport) PURE_VIRTUAL(SetTeleportToExplicitTime, )
+
+	/** Set what to do when SequenceEvaluator is reinitialized */
+	virtual void SetReinitializationBehavior(TEnumAsByte<ESequenceEvalReinit::Type> InBehavior) PURE_VIRTUAL(SetReinitializationBehavior, )
+
 	// The animation sequence asset to evaluate
-	UPROPERTY(EditAnywhere, BlueprintReadWrite, Category = Settings, meta = (PinHiddenByDefault))
-	TObjectPtr<UAnimSequenceBase> Sequence;
+	virtual UAnimSequenceBase* GetSequence() const PURE_VIRTUAL(GetSequence, return nullptr;)
 
 	// The time at which to evaluate the associated sequence
-	UPROPERTY(EditAnywhere, BlueprintReadWrite, Category=Settings, meta=(PinShownByDefault))
-	float ExplicitTime;
+	virtual float GetExplicitTime() const PURE_VIRTUAL(GetExplicitTime,  return 0.0f;)
 
 	/** This only works if bTeleportToExplicitTime is false OR this node is set to use SyncGroup */
-	UPROPERTY(EditAnywhere, BlueprintReadWrite, Category = Settings, meta = (PinHiddenByDefault))
-	bool bShouldLoop;
+	virtual bool GetShouldLoop() const PURE_VIRTUAL(GetShouldLoop, return true;)
 
 	/** If true, teleport to explicit time, does NOT advance time (does not trigger notifies, does not extract Root Motion, etc.)
 	If false, will advance time (will trigger notifies, extract root motion if applicable, etc.)
 	Note: using a sync group forces advancing time regardless of what this option is set to. */
-	UPROPERTY(EditAnywhere, BlueprintReadWrite, Category = Settings, meta = (PinHiddenByDefault))
-	bool bTeleportToExplicitTime;
+	virtual bool GetTeleportToExplicitTime() const PURE_VIRTUAL(GetTeleportToExplicitTime, return true;)
 
 	/** What to do when SequenceEvaluator is reinitialized */
-	UPROPERTY(EditAnywhere, BlueprintReadWrite, Category = Settings, meta=(DisplayAfter="StartPosition"))
-	TEnumAsByte<ESequenceEvalReinit::Type> ReinitializationBehavior;
+	virtual TEnumAsByte<ESequenceEvalReinit::Type> GetReinitializationBehavior() const PURE_VIRTUAL(GetReinitializationBehavior, return ESequenceEvalReinit::ExplicitTime;)
+
+	// The start up position, it only applies when ReinitializationBehavior == StartPosition. Only used when bTeleportToExplicitTime is false.
+	virtual float GetStartPosition() const PURE_VIRTUAL(GetStartPosition,  return 0.0f;)
+};
+
+// Sequence evaluator node that can be used with constant folding
+USTRUCT(BlueprintInternalUseOnly)
+struct ANIMGRAPHRUNTIME_API FAnimNode_SequenceEvaluator : public FAnimNode_SequenceEvaluatorBase
+{
+	GENERATED_BODY()
+
+private:
+	friend class UAnimGraphNode_SequenceEvaluator;
+
+#if WITH_EDITORONLY_DATA
+	// The animation sequence asset to evaluate
+	UPROPERTY(EditAnywhere, Category = Settings, meta = (PinHiddenByDefault, FoldProperty))
+	TObjectPtr<UAnimSequenceBase> Sequence = nullptr;
+
+	// The time at which to evaluate the associated sequence
+	UPROPERTY(EditAnywhere, Category=Settings, meta=(PinShownByDefault, FoldProperty))
+	float ExplicitTime = 0.0f;
+
+	/** This only works if bTeleportToExplicitTime is false OR this node is set to use SyncGroup */
+	UPROPERTY(EditAnywhere, Category = Settings, meta = (PinHiddenByDefault, FoldProperty))
+	bool bShouldLoop = true;
 
-	bool bReinitialized;
+	/** If true, teleport to explicit time, does NOT advance time (does not trigger notifies, does not extract Root Motion, etc.)
+	If false, will advance time (will trigger notifies, extract root motion if applicable, etc.)
+	Note: using a sync group forces advancing time regardless of what this option is set to. */
+	UPROPERTY(EditAnywhere, Category = Settings, meta = (PinHiddenByDefault, FoldProperty))
+	bool bTeleportToExplicitTime = true;
+
+	/** What to do when SequenceEvaluator is reinitialized */
+	UPROPERTY(EditAnywhere, Category = Settings, meta=(DisplayAfter="StartPosition", FoldProperty))
+	TEnumAsByte<ESequenceEvalReinit::Type> ReinitializationBehavior = ESequenceEvalReinit::ExplicitTime;
 
 	// The start up position, it only applies when ReinitializationBehavior == StartPosition. Only used when bTeleportToExplicitTime is false.
-	UPROPERTY(EditAnywhere, BlueprintReadWrite, Category = Settings, meta = (PinHiddenByDefault))
-	float StartPosition;
+	UPROPERTY(EditAnywhere, Category = Settings, meta = (PinHiddenByDefault, FoldProperty))
+	float StartPosition = 0.0f;
+#endif
 
 public:
-	FAnimNode_SequenceEvaluator()
-		: Sequence(nullptr)
-		, ExplicitTime(0.0f)
-		, bShouldLoop(false)
-		, bTeleportToExplicitTime(true)
-		, ReinitializationBehavior(ESequenceEvalReinit::ExplicitTime)
-		, bReinitialized(false)
-		, StartPosition(0.f)
-	{
-	}
+	// FAnimNode_SequenceEvaluatorBase interface
+	virtual void SetSequence(UAnimSequenceBase* InSequence) override;
+	virtual UAnimSequenceBase* GetSequence() const override;
+	virtual float GetExplicitTime() const override;
+	virtual bool GetShouldLoop() const override;
+	virtual bool GetTeleportToExplicitTime() const override;
+	virtual TEnumAsByte<ESequenceEvalReinit::Type> GetReinitializationBehavior() const override;
+	virtual float GetStartPosition() const override;
+};
 
-	// FAnimNode_AssetPlayerBase interface
-	virtual float GetCurrentAssetTime();
-	virtual float GetCurrentAssetLength();
-	// End of FAnimNode_AssetPlayerBase interface
+// Sequence evaluator node that can be used standalone (without constant folding)
+USTRUCT(BlueprintInternalUseOnly)
+struct ANIMGRAPHRUNTIME_API FAnimNode_SequenceEvaluator_Standalone : public FAnimNode_SequenceEvaluatorBase
+{
+	GENERATED_BODY()
 
-	// FAnimNode_Base interface
-	virtual void Initialize_AnyThread(const FAnimationInitializeContext& Context) override;
-	virtual void CacheBones_AnyThread(const FAnimationCacheBonesContext& Context) override;
-	virtual void UpdateAssetPlayer(const FAnimationUpdateContext& Context) override;
-	virtual void Evaluate_AnyThread(FPoseContext& Output) override;
-	virtual void OverrideAsset(UAnimationAsset* NewAsset) override;
-	virtual void GatherDebugData(FNodeDebugData& DebugData) override;
-	// End of FAnimNode_Base interface
+private:
+	// The animation sequence asset to evaluate
+	UPROPERTY(EditAnywhere, Category = Settings, meta = (PinHiddenByDefault))
+	TObjectPtr<UAnimSequenceBase> Sequence = nullptr;
 
-	// FAnimNode_AssetPlayerBase Interface
-	virtual float GetAccumulatedTime() const {return ExplicitTime;}
-	virtual void SetAccumulatedTime(const float& NewTime) {ExplicitTime = NewTime;}
-	virtual UAnimationAsset* GetAnimAsset() {return Sequence;}
-	// End of FAnimNode_AssetPlayerBase Interface
+	// The time at which to evaluate the associated sequence
+	UPROPERTY(EditAnywhere, Category=Settings, meta=(PinShownByDefault))
+	float ExplicitTime = 0.0f;
 
-	void SetExplicitPreviousTime(float PreviousTime) { InternalTimeAccumulator = PreviousTime; }
-};
+	/** This only works if bTeleportToExplicitTime is false OR this node is set to use SyncGroup */
+	UPROPERTY(EditAnywhere, Category = Settings, meta = (PinHiddenByDefault))
+	bool bShouldLoop = true;
+
+	/** If true, teleport to explicit time, does NOT advance time (does not trigger notifies, does not extract Root Motion, etc.)
+	If false, will advance time (will trigger notifies, extract root motion if applicable, etc.)
+	Note: using a sync group forces advancing time regardless of what this option is set to. */
+	UPROPERTY(EditAnywhere, Category = Settings, meta = (PinHiddenByDefault))
+	bool bTeleportToExplicitTime = true;
+
+	/** What to do when SequenceEvaluator is reinitialized */
+	UPROPERTY(EditAnywhere, Category = Settings, meta=(DisplayAfter="StartPosition"))
+	TEnumAsByte<ESequenceEvalReinit::Type> ReinitializationBehavior = ESequenceEvalReinit::ExplicitTime;
+
+	// The start up position, it only applies when ReinitializationBehavior == StartPosition. Only used when bTeleportToExplicitTime is false.
+	UPROPERTY(EditAnywhere, Category = Settings, meta = (PinHiddenByDefault))
+	float StartPosition = 0.0f;
+
+public:
+	// FAnimNode_SequenceEvaluatorBase interface
+	virtual void SetSequence(UAnimSequenceBase* InSequence) override { Sequence = InSequence; }
+	virtual void SetExplicitTime(float InTime) override { ExplicitTime = InTime; }
+	virtual void SetTeleportToExplicitTime(bool bInTeleport) override { bTeleportToExplicitTime = bInTeleport; }
+	virtual void SetReinitializationBehavior(TEnumAsByte<ESequenceEvalReinit::Type> InBehavior) override { ReinitializationBehavior = InBehavior; }
+	virtual UAnimSequenceBase* GetSequence() const override { return Sequence; }
+	virtual float GetExplicitTime() const override { return ExplicitTime; }
+	virtual bool GetShouldLoop() const override { return bShouldLoop; }
+	virtual bool GetTeleportToExplicitTime() const override { return bTeleportToExplicitTime; }
+	virtual TEnumAsByte<ESequenceEvalReinit::Type> GetReinitializationBehavior() const override { return ReinitializationBehavior; }
+	virtual float GetStartPosition() const override { return StartPosition; }
+};
\ No newline at end of file
diff --git a/Engine/Source/Runtime/AnimGraphRuntime/Public/AnimSequencerInstanceProxy.h b/Engine/Source/Runtime/AnimGraphRuntime/Public/AnimSequencerInstanceProxy.h
--- a/Engine/Source/Runtime/AnimGraphRuntime/Public/AnimSequencerInstanceProxy.h
+++ b/Engine/Source/Runtime/AnimGraphRuntime/Public/AnimSequencerInstanceProxy.h
@@ -65,7 +65,7 @@ struct FSequencerPlayerAnimSequence : public FSequencerPlayerBase
 {
 	SEQUENCER_INSTANCE_PLAYER_TYPE(FSequencerPlayerAnimSequence, FSequencerPlayerBase)
 	TOptional<FRootMotionOverride> RootMotion;
-	struct FAnimNode_SequenceEvaluator PlayerNode;
+	FAnimNode_SequenceEvaluator_Standalone PlayerNode;
 };
 
 
diff --git a/Engine/Source/Runtime/Core/Public/UObject/UE5MainStreamObjectVersion.h b/Engine/Source/Runtime/Core/Public/UObject/UE5MainStreamObjectVersion.h
--- a/Engine/Source/Runtime/Core/Public/UObject/UE5MainStreamObjectVersion.h
+++ b/Engine/Source/Runtime/Core/Public/UObject/UE5MainStreamObjectVersion.h
@@ -115,6 +115,9 @@ struct CORE_API FUE5MainStreamObjectVersion
 		// PackedLevelInstance bounds fix
 		PackedLevelInstanceBoundsFix,
 
+		// Initial set of anim nodes converted to use constants held in sparse class data
+		AnimNodeConstantDataRefactorPhase0,
+
         // -----<new versions can be added above this line>-------------------------------------------------
 		VersionPlusOne,
 		LatestVersion = VersionPlusOne - 1
diff --git a/Engine/Source/Runtime/Engine/Classes/Animation/AnimBlueprint.h b/Engine/Source/Runtime/Engine/Classes/Animation/AnimBlueprint.h
--- a/Engine/Source/Runtime/Engine/Classes/Animation/AnimBlueprint.h
+++ b/Engine/Source/Runtime/Engine/Classes/Animation/AnimBlueprint.h
@@ -144,6 +144,9 @@ class ENGINE_API UAnimBlueprint : public UBlueprint, public IInterface_PreviewMe
 	/** Returns the most base anim blueprint for a given blueprint (if it is inherited from another anim blueprint, returning null if only native / non-anim BP classes are it's parent) */
 	static UAnimBlueprint* FindRootAnimBlueprint(const UAnimBlueprint* DerivedBlueprint);
 
+	/** Returns the parent anim blueprint for a given blueprint (if it is inherited from another anim blueprint, returning null if only native / non-anim BP classes are it's parent) */
+	static UAnimBlueprint* GetParentAnimBlueprint(const UAnimBlueprint* DerivedBlueprint);
+	
 	DECLARE_MULTICAST_DELEGATE_TwoParams(FOnOverrideChangedMulticaster, FGuid, UAnimationAsset*);
 
 	typedef FOnOverrideChangedMulticaster::FDelegate FOnOverrideChanged;
@@ -165,7 +168,7 @@ class ENGINE_API UAnimBlueprint : public UBlueprint, public IInterface_PreviewMe
 
 	virtual void PostLoad() override;
 	virtual bool FindDiffs(const UBlueprint* OtherBlueprint, FDiffResults& Results) const override;
-
+	
 protected:
 	// Broadcast when an override is changed, allowing derived blueprints to be updated
 	FOnOverrideChangedMulticaster OnOverrideChanged;
@@ -195,6 +198,11 @@ public:
 	FAnimBlueprintDebugData* GetDebugData() const;
 
 #if WITH_EDITORONLY_DATA
+public:
+	// Queue a refresh of the set of anim blueprint extensions that this anim blueprint hosts.
+	// Usually called from anim graph nodes to ensure that extensions that are no longer required are cleaned up.
+	void RequestRefreshExtensions() { bRefreshExtensions = true; }
+	
 public:
 	// Array of overrides to asset containing nodes in the parent that have been overridden
 	UPROPERTY()
@@ -206,6 +214,8 @@ public:
 	TArray<TObjectPtr<class UPoseWatch>> PoseWatches;
 
 private:
+	friend class FAnimBlueprintCompilerContext;
+	
 	/** The default skeletal mesh to use when previewing this asset - this only applies when you open Persona using this asset*/
 	UPROPERTY(duplicatetransient, AssetRegistrySearchable)
 	TSoftObjectPtr<class USkeletalMesh> PreviewSkeletalMesh;
@@ -224,5 +234,8 @@ private:
 	/** The tag to use when applying a preview animation blueprint via LinkAnimGraphByTag */
 	UPROPERTY()
 	FName PreviewAnimationBlueprintTag;
+
+	/** If set, then extensions need to be refreshed according to spawned nodes */
+	bool bRefreshExtensions;
 #endif // WITH_EDITORONLY_DATA
 };
