Backing out CL 16071104

diff --git a/Engine/Source/Editor/AnimGraph/Public/IAnimBlueprintCompilerCreationContext.h b/Engine/Source/Editor/AnimGraph/Public/IAnimBlueprintCompilerCreationContext.h
--- a/Engine/Source/Editor/AnimGraph/Public/IAnimBlueprintCompilerCreationContext.h
+++ b/Engine/Source/Editor/AnimGraph/Public/IAnimBlueprintCompilerCreationContext.h
@@ -17,51 +17,83 @@ class UAnimGraphNode_Base;
 class UEdGraph;
 class UEdGraphSchema;
 
-UE_DEPRECATED(5.0, "Anim BP compiler delegate system is deprecated. Use UAnimBlueprintExtension instead")
+/** Begin ordered delegate calls - these functions are called in the order presented here */
+
+/** 
+ * Delegate fired when the class starts compiling. The class may be new or recycled.
+ * @param	InClass			The class that is being compiled. This could be a newly created class or a recycled class
+ * @param	InCompilerContext	The compiler context for the current compilation
+ * @param	OutCompiledData	The compiled data that this handler can write to
+ */
 DECLARE_MULTICAST_DELEGATE_ThreeParams(FOnStartCompilingClass, const UClass* /*InClass*/, IAnimBlueprintCompilationBracketContext& /*InCompilerContext*/, IAnimBlueprintGeneratedClassCompiledData& /*OutCompiledData*/);
 
-UE_DEPRECATED(5.0, "Anim BP compiler delegate system is deprecated. Use UAnimBlueprintExtension instead")
+/** 
+ * Delegate fired before all animation nodes are processed
+ * @param	InAnimNodes		The set of anim nodes that should be processed. Note that these nodes have not yet been pruned for connectivity
+ * @param	InCompilerContext	The compiler context for the current compilation
+ * @param	OutCompiledData	The compiled data that this handler can write to
+ */
 DECLARE_MULTICAST_DELEGATE_ThreeParams(FOnPreProcessAnimationNodes, TArrayView<UAnimGraphNode_Base*> /*InAnimNodes*/, IAnimBlueprintCompilationContext& /*InCompilerContext*/, IAnimBlueprintGeneratedClassCompiledData& /*OutCompiledData*/);
 
-UE_DEPRECATED(5.0, "Anim BP compiler delegate system is deprecated. Use UAnimBlueprintExtension instead")
+/** 
+ * Delegate fired after all animation nodes are processed
+ * @param	InAnimNodes		The set of anim nodes that were processed. Note that these nodes were not pruned for connectivity (they will be the same set passed to PreProcessAnimationNodes)
+ * @param	InCompilerContext	The compiler context for the current compilation
+ * @param	OutCompiledData	The compiled data that this handler can write to
+ */
 DECLARE_MULTICAST_DELEGATE_ThreeParams(FOnPostProcessAnimationNodes, TArrayView<UAnimGraphNode_Base*> /*InAnimNodes*/, IAnimBlueprintCompilationContext& /*InCompilerContext*/, IAnimBlueprintGeneratedClassCompiledData& /*OutCompiledData*/);
 
-UE_DEPRECATED(5.0, "Anim BP compiler delegate system is deprecated. Use UAnimBlueprintExtension instead")
+/** 
+ * Delegate fired post- graph expansion
+ * @param	InGraph		The graph that was just expanded
+ * @param	InCompilerContext	The compiler context for the current compilation
+ * @param	OutCompiledData	The compiled data that this handler can write to
+ */
 DECLARE_MULTICAST_DELEGATE_ThreeParams(FOnPostExpansionStep, const UEdGraph* /*InGraph*/, IAnimBlueprintPostExpansionStepContext& /*InCompilerContext*/, IAnimBlueprintGeneratedClassCompiledData& /*OutCompiledData*/);
 
-UE_DEPRECATED(5.0, "Anim BP compiler delegate system is deprecated. Use UAnimBlueprintExtension instead")
+/** 
+ * Delegate fired when the class has finished compiling
+ * @param	InClass			The class that was compiled
+ * @param	InCompilerContext	The compiler context for the current compilation
+ * @param	OutCompiledData	The compiled data that this handler can write to
+ */
 DECLARE_MULTICAST_DELEGATE_ThreeParams(FOnFinishCompilingClass, const UClass* /*InClass*/, IAnimBlueprintCompilationBracketContext& /*InCompilerContext*/, IAnimBlueprintGeneratedClassCompiledData& /*OutCompiledData*/);
 
-UE_DEPRECATED(5.0, "Anim BP compiler delegate system is deprecated. Use UAnimBlueprintExtension instead")
+/** 
+ * Delegate fired when data is being copied to the CDO
+ * @param	InDefaultObject		The CDO for the just-compiled class
+ * @param	InCompilerContext	The compiler context for the current compilation
+ * @param	OutCompiledData	The compiled data that this handler can write to
+ */
 DECLARE_MULTICAST_DELEGATE_ThreeParams(FOnCopyTermDefaultsToDefaultObject, UObject* /*InDefaultObject*/, IAnimBlueprintCopyTermDefaultsContext& /*InCompilerContext*/, IAnimBlueprintGeneratedClassCompiledData& /*OutCompiledData*/);
 
+/** End ordered calls */
+
 /** Interface to the creation of an anim BP compiler context */
 class ANIMGRAPH_API IAnimBlueprintCompilerCreationContext
 {
 public:
 	virtual ~IAnimBlueprintCompilerCreationContext() {}
 
-	/** Registers a graphs schema class to the anim BP compiler so that default function processing is not performed on it */
-	virtual void RegisterKnownGraphSchema(TSubclassOf<UEdGraphSchema> InGraphSchemaClass) = 0;
-	
-	PRAGMA_DISABLE_DEPRECATION_WARNINGS
-	UE_DEPRECATED(5.0, "Anim BP compiler delegate system is deprecated. Use UAnimBlueprintExtension instead")
-	virtual FOnStartCompilingClass& OnStartCompilingClass() { static FOnStartCompilingClass Dummy; return Dummy; }
+	/** Delegate fired when the class starts compiling. The class may be new or recycled. */
+	virtual FOnStartCompilingClass& OnStartCompilingClass() = 0;
 
-	UE_DEPRECATED(5.0, "Anim BP compiler delegate system is deprecated. Use UAnimBlueprintExtension instead")
-	virtual FOnPreProcessAnimationNodes& OnPreProcessAnimationNodes() { static FOnPreProcessAnimationNodes Dummy; return Dummy; }
+	/** Delegate fired before all animation nodes are processed */
+	virtual FOnPreProcessAnimationNodes& OnPreProcessAnimationNodes() = 0;
 
-	UE_DEPRECATED(5.0, "Anim BP compiler delegate system is deprecated. Use UAnimBlueprintExtension instead")
-	virtual FOnPostProcessAnimationNodes& OnPostProcessAnimationNodes() { static FOnPostProcessAnimationNodes Dummy; return Dummy; }
+	/** Delegate fired after all animation nodes are processed */
+	virtual FOnPostProcessAnimationNodes& OnPostProcessAnimationNodes() = 0;
 
-	UE_DEPRECATED(5.0, "Anim BP compiler delegate system is deprecated. Use UAnimBlueprintExtension instead")
-	virtual FOnPostExpansionStep& OnPostExpansionStep() { static FOnPostExpansionStep Dummy; return Dummy; }
+	/** Delegate fired post- graph expansion */
+	virtual FOnPostExpansionStep& OnPostExpansionStep() = 0;
 
-	UE_DEPRECATED(5.0, "Anim BP compiler delegate system is deprecated. Use UAnimBlueprintExtension instead")
-	virtual FOnFinishCompilingClass& OnFinishCompilingClass() { static FOnFinishCompilingClass Dummy; return Dummy; }
+	/** Delegate fired when the class has finished compiling */
+	virtual FOnFinishCompilingClass& OnFinishCompilingClass() = 0;
 
-	UE_DEPRECATED(5.0, "Anim BP compiler delegate system is deprecated. Use UAnimBlueprintExtension instead")
-	virtual FOnCopyTermDefaultsToDefaultObject& OnCopyTermDefaultsToDefaultObject() { static FOnCopyTermDefaultsToDefaultObject Dummy; return Dummy; }
-	PRAGMA_ENABLE_DEPRECATION_WARNINGS
+	/** Delegate fired when data is being copied to the CDO */
+	virtual FOnCopyTermDefaultsToDefaultObject& OnCopyTermDefaultsToDefaultObject() = 0;
+
+	/** Registers a graphs schema class to the anim BP compiler so that default function processing is not performed on it */
+	virtual void RegisterKnownGraphSchema(TSubclassOf<UEdGraphSchema> InGraphSchemaClass) = 0;
 };
 
diff --git a/Engine/Source/Editor/AnimGraph/Public/IAnimBlueprintCompilerHandler.h b/Engine/Source/Editor/AnimGraph/Public/IAnimBlueprintCompilerHandler.h
--- a/Engine/Source/Editor/AnimGraph/Public/IAnimBlueprintCompilerHandler.h
+++ b/Engine/Source/Editor/AnimGraph/Public/IAnimBlueprintCompilerHandler.h
@@ -4,8 +4,6 @@
 
 #include "CoreMinimal.h"
 
-class UE_DEPRECATED(5.0, "IAnimBlueprintCompilerHandler is deprecated, use UAnimBlueprintExtension_PropertyAccess instead") IAnimBlueprintCompilerHandler;
-
 /** Empty base class for handlers that are maintained for the lifetime of the compiler context */
 class ANIMGRAPH_API IAnimBlueprintCompilerHandler
 {
diff --git a/Engine/Source/Editor/AnimGraph/Public/IAnimBlueprintCompilerHandlerCollection.h b/Engine/Source/Editor/AnimGraph/Public/IAnimBlueprintCompilerHandlerCollection.h
--- a/Engine/Source/Editor/AnimGraph/Public/IAnimBlueprintCompilerHandlerCollection.h
+++ b/Engine/Source/Editor/AnimGraph/Public/IAnimBlueprintCompilerHandlerCollection.h
@@ -8,19 +8,16 @@
 class FAnimBlueprintCompilerContext;
 class IAnimBlueprintCompilerCreationContext;
 
-class UE_DEPRECATED(5.0, "IAnimBlueprintCompilerHandlerCollection is no longer used") IAnimBlueprintCompilerHandlerCollection;
-
 class ANIMGRAPH_API IAnimBlueprintCompilerHandlerCollection
 {
 public:
-	PRAGMA_DISABLE_DEPRECATION_WARNINGS
 	virtual ~IAnimBlueprintCompilerHandlerCollection() {}
 
 	/** Register a named handler */
-	static void RegisterHandler(FName InName, TFunction<TUniquePtr<IAnimBlueprintCompilerHandler>(IAnimBlueprintCompilerCreationContext&)> InFunction) {}
+	static void RegisterHandler(FName InName, TFunction<TUniquePtr<IAnimBlueprintCompilerHandler>(IAnimBlueprintCompilerCreationContext&)> InFunction);
 
 	/** Register a named handler */
-	static void UnregisterHandler(FName InName) {}
+	static void UnregisterHandler(FName InName);
 
 	// Get a GetHandler of specified type with the specified name
 	template<typename THandlerClass>
@@ -32,5 +29,4 @@ public:
 protected:
 	// Get a GetHandler with the specified name
 	virtual IAnimBlueprintCompilerHandler* GetHandlerByName(FName InName) const = 0;
-	PRAGMA_ENABLE_DEPRECATION_WARNINGS
 };
\ No newline at end of file
diff --git a/Engine/Source/Editor/AnimGraph/Public/IAnimBlueprintCopyTermDefaultsContext.h b/Engine/Source/Editor/AnimGraph/Public/IAnimBlueprintCopyTermDefaultsContext.h
--- a/Engine/Source/Editor/AnimGraph/Public/IAnimBlueprintCopyTermDefaultsContext.h
+++ b/Engine/Source/Editor/AnimGraph/Public/IAnimBlueprintCopyTermDefaultsContext.h
@@ -11,7 +11,7 @@ class UAnimBlueprint;
 class IAnimBlueprintCopyTermDefaultsContext
 {
 public:	
-	virtual ~IAnimBlueprintCopyTermDefaultsContext() = default;
+	virtual ~IAnimBlueprintCopyTermDefaultsContext() {}
 
 	// Get the message log for the current compilation
 	FCompilerResultsLog& GetMessageLog() const { return GetMessageLogImpl(); }
@@ -32,7 +32,7 @@ protected:
 class IAnimBlueprintNodeCopyTermDefaultsContext
 {
 public:	
-	virtual ~IAnimBlueprintNodeCopyTermDefaultsContext() = default;
+	virtual ~IAnimBlueprintNodeCopyTermDefaultsContext() {}
 
 	// Get the property that we are writing to
 	const FProperty* GetTargetProperty() const { return GetTargetPropertyImpl(); }
@@ -58,36 +58,4 @@ protected:
 
 	// Get the property index for this node
 	virtual int32 GetNodePropertyIndexImpl() const  = 0;
-};
-
-// Interface passed to per-extension CopyTermDefaults override point
-class IAnimBlueprintExtensionCopyTermDefaultsContext
-{
-public:	
-	virtual ~IAnimBlueprintExtensionCopyTermDefaultsContext() = default;
-
-	// Get the property that we are writing to
-	const FProperty* GetTargetProperty() const { return GetTargetPropertyImpl(); }
-
-	// Get the destination ptr (the node) that we are writing to
-	uint8* GetDestinationPtr() const { return GetDestinationPtrImpl(); }
-
-	// Get the source ptr (the node in the anim graph node) that we are reading from
-	const uint8* GetSourcePtr() const { return GetSourcePtrImpl(); }
-
-	// Get the property index for this node
-	int32 GetNodePropertyIndex() const { return GetNodePropertyIndexImpl(); }
-
-protected:
-	// Get the property that we are writing to
-	virtual const FProperty* GetTargetPropertyImpl() const = 0;
-
-	// Get the destination ptr (the node) that we are writing to
-	virtual uint8* GetDestinationPtrImpl() const = 0;
-
-	// Get the source ptr (the node in the anim graph node) that we are reading from
-	virtual const uint8* GetSourcePtrImpl() const = 0;
-
-	// Get the property index for this node
-	virtual int32 GetNodePropertyIndexImpl() const = 0;
-};
+};
\ No newline at end of file
diff --git a/Engine/Source/Editor/AnimGraph/Public/IAnimBlueprintGeneratedClassCompiledData.h b/Engine/Source/Editor/AnimGraph/Public/IAnimBlueprintGeneratedClassCompiledData.h
--- a/Engine/Source/Editor/AnimGraph/Public/IAnimBlueprintGeneratedClassCompiledData.h
+++ b/Engine/Source/Editor/AnimGraph/Public/IAnimBlueprintGeneratedClassCompiledData.h
@@ -7,9 +7,9 @@
 #include "Animation/AnimClassInterface.h"
 #include "Animation/AnimTypes.h"
 #include "Animation/AnimNodeBase.h"
-#include "PropertyAccess.h"
 
 struct FBlueprintDebugData;
+struct FPropertyAccessLibrary;
 struct FAnimBlueprintDebugData;
 
 /** Interface to the writable parts of the generated class that handlers can operate on */
@@ -31,11 +31,11 @@ public:
 	// Finds a notify event or adds if it doesn't already exist
 	virtual int32 FindOrAddNotify(FAnimNotifyEvent& Notify) const = 0;
 
-	UE_DEPRECATED(5.0, "Exposed value handlers are no longer accessible via this method. Access them via UAnimBlueprintExtension_Base.")
-	virtual TArray<FExposedValueHandler>& GetExposedValueHandlers() const {	static TArray<FExposedValueHandler> Dummy; return Dummy; }
+	// Get the currently-compiled classes exposed value handlers
+	virtual TArray<FExposedValueHandler>& GetExposedValueHandlers() const = 0;
 
-	UE_DEPRECATED(5.0, "Property access library is no longer accessible via this method. Access it via UAnimBlueprintExtension_PropertyAccess.")
-	virtual FPropertyAccessLibrary& GetPropertyAccessLibrary() const { static FPropertyAccessLibrary Dummy; return Dummy; }
+	// Get the currently-compiled classes property access library
+	virtual FPropertyAccessLibrary& GetPropertyAccessLibrary() const = 0;
 
 	// Get the anim debug data for the currently-compiled class
 	virtual FAnimBlueprintDebugData& GetAnimBlueprintDebugData() const = 0;
diff --git a/Engine/Source/Editor/AnimGraph/Public/IAnimBlueprintPostExpansionStepContext.h b/Engine/Source/Editor/AnimGraph/Public/IAnimBlueprintPostExpansionStepContext.h
--- a/Engine/Source/Editor/AnimGraph/Public/IAnimBlueprintPostExpansionStepContext.h
+++ b/Engine/Source/Editor/AnimGraph/Public/IAnimBlueprintPostExpansionStepContext.h
@@ -3,6 +3,7 @@
 #pragma once
 
 #include "CoreMinimal.h"
+#include "IAnimBlueprintCompilerHandler.h"
 
 class FCompilerResultsLog;
 class UEdGraph;
@@ -14,6 +15,14 @@ class IAnimBlueprintPostExpansionStepContext
 public:
 	virtual ~IAnimBlueprintPostExpansionStepContext() {}
 
+	// Get a handler of the specified type and name (i.e. via simple name-based RTTI)
+	// Handlers are registered via IAnimBlueprintCompilerHandlerCollection::RegisterHandler
+	template <typename THandlerClass>
+	THandlerClass* GetHandler(FName InName) const
+	{
+		return static_cast<THandlerClass*>(GetHandlerInternal(InName));
+	}
+
 	// Get the message log for the current compilation
 	FCompilerResultsLog& GetMessageLog() const { return GetMessageLogImpl(); }
 
@@ -32,4 +41,7 @@ protected:
 
 	// Get the compiler options we are currently using
 	virtual const FKismetCompilerOptions& GetCompileOptionsImpl() const = 0;
+
+	// GetHandler helper function
+	virtual IAnimBlueprintCompilerHandler* GetHandlerInternal(FName InName) const = 0;
 };
diff --git a/Engine/Source/Editor/AnimGraph/Public/PropertyAccessCompilerHandler.h b/Engine/Source/Editor/AnimGraph/Public/PropertyAccessCompilerHandler.h
--- a/Engine/Source/Editor/AnimGraph/Public/PropertyAccessCompilerHandler.h
+++ b/Engine/Source/Editor/AnimGraph/Public/PropertyAccessCompilerHandler.h
@@ -10,14 +10,10 @@
 
 enum class EPropertyAccessBatchType : uint8;
 class IAnimBlueprintGeneratedClassCompiledData;
-class IAnimBlueprintCompilationBracketContext;
 
 // Delegate called when the library is compiled (whether successfully or not)
-DECLARE_MULTICAST_DELEGATE_TwoParams(FOnPostLibraryCompiled, IAnimBlueprintCompilationBracketContext& /*InCompilationContext*/, IAnimBlueprintGeneratedClassCompiledData& /*OutCompiledData*/)
+DECLARE_MULTICAST_DELEGATE_OneParam(FOnPostLibraryCompiled, IAnimBlueprintGeneratedClassCompiledData& /*OutCompiledData*/)
 
-class UE_DEPRECATED(5.0, "FPropertyAccessCompilerHandler is no longer used. Use UAnimBlueprintExtension_PropertyAccess instead") FPropertyAccessCompilerHandler;
-
-PRAGMA_DISABLE_DEPRECATION_WARNINGS
 class FPropertyAccessCompilerHandler : public IAnimBlueprintCompilerHandler
 {
 public:
@@ -33,5 +29,4 @@ public:
 
 	// Maps the initial integer copy handle to a true handle, post compilation
 	virtual int32 MapCopyIndex(int32 InIndex) const = 0;
-};
-PRAGMA_ENABLE_DEPRECATION_WARNINGS
\ No newline at end of file
+};
\ No newline at end of file
diff --git a/Engine/Source/Editor/AnimationBlueprintEditor/AnimationBlueprintEditor.Build.cs b/Engine/Source/Editor/AnimationBlueprintEditor/AnimationBlueprintEditor.Build.cs
--- a/Engine/Source/Editor/AnimationBlueprintEditor/AnimationBlueprintEditor.Build.cs
+++ b/Engine/Source/Editor/AnimationBlueprintEditor/AnimationBlueprintEditor.Build.cs
@@ -42,7 +42,6 @@ public class AnimationBlueprintEditor : ModuleRules
                 "RHI",
                 "KismetCompiler",
 				"ToolMenus",
-				"AnimGraphRuntime",
             }
 		);
 
diff --git a/Engine/Source/Editor/AnimationBlueprintEditor/Private/AnimationBlueprintEditor.cpp b/Engine/Source/Editor/AnimationBlueprintEditor/Private/AnimationBlueprintEditor.cpp
--- a/Engine/Source/Editor/AnimationBlueprintEditor/Private/AnimationBlueprintEditor.cpp
+++ b/Engine/Source/Editor/AnimationBlueprintEditor/Private/AnimationBlueprintEditor.cpp
@@ -758,7 +758,7 @@ void FAnimationBlueprintEditor::OnConvertToSequenceEvaluator()
 			UAnimGraphNode_SequencePlayer* OldNode = Cast<UAnimGraphNode_SequencePlayer>(*NodeIter);
 
 			// see if sequence player
-			if ( OldNode && OldNode->Node.GetSequence() )
+			if ( OldNode && OldNode->Node.Sequence )
 			{
 				UEdGraph* TargetGraph = OldNode->GetGraph();
 				TargetGraph->Modify();
@@ -767,7 +767,7 @@ void FAnimationBlueprintEditor::OnConvertToSequenceEvaluator()
 				// create new evaluator
 				FGraphNodeCreator<UAnimGraphNode_SequenceEvaluator> NodeCreator(*TargetGraph);
 				UAnimGraphNode_SequenceEvaluator* NewNode = NodeCreator.CreateNode();
-				NewNode->Node.SetSequence(OldNode->Node.GetSequence());
+				NewNode->Node.Sequence = OldNode->Node.Sequence;
 				NodeCreator.Finalize();
 
 				// get default data from old node to new node
@@ -807,7 +807,7 @@ void FAnimationBlueprintEditor::OnConvertToSequencePlayer()
 			UAnimGraphNode_SequenceEvaluator* OldNode = Cast<UAnimGraphNode_SequenceEvaluator>(*NodeIter);
 
 			// see if sequence player
-			if ( OldNode && OldNode->Node.GetSequence() )
+			if ( OldNode && OldNode->Node.Sequence )
 			{
 				// convert to sequence player
 				UEdGraph* TargetGraph = OldNode->GetGraph();
@@ -817,7 +817,7 @@ void FAnimationBlueprintEditor::OnConvertToSequencePlayer()
 				// create new player
 				FGraphNodeCreator<UAnimGraphNode_SequencePlayer> NodeCreator(*TargetGraph);
 				UAnimGraphNode_SequencePlayer* NewNode = NodeCreator.CreateNode();
-				NewNode->Node.SetSequence(OldNode->Node.GetSequence());
+				NewNode->Node.Sequence = OldNode->Node.Sequence;
 				NodeCreator.Finalize();
 
 				// get default data from old node to new node
@@ -858,7 +858,7 @@ void FAnimationBlueprintEditor::OnConvertToBlendSpaceEvaluator()
 			UAnimGraphNode_BlendSpacePlayer* OldNode = Cast<UAnimGraphNode_BlendSpacePlayer>(*NodeIter);
 
 			// see if sequence player
-			if ( OldNode && OldNode->Node.GetBlendSpace() )
+			if ( OldNode && OldNode->Node.BlendSpace )
 			{
 				// convert to sequence evaluator
 				UEdGraph* TargetGraph = OldNode->GetGraph();
@@ -868,7 +868,7 @@ void FAnimationBlueprintEditor::OnConvertToBlendSpaceEvaluator()
 				// create new evaluator
 				FGraphNodeCreator<UAnimGraphNode_BlendSpaceEvaluator> NodeCreator(*TargetGraph);
 				UAnimGraphNode_BlendSpaceEvaluator* NewNode = NodeCreator.CreateNode();
-				NewNode->Node.SetBlendSpace(OldNode->Node.GetBlendSpace());
+				NewNode->Node.BlendSpace = OldNode->Node.BlendSpace;
 				NodeCreator.Finalize();
 
 				// get default data from old node to new node
@@ -909,7 +909,7 @@ void FAnimationBlueprintEditor::OnConvertToBlendSpacePlayer()
 			UAnimGraphNode_BlendSpaceEvaluator* OldNode = Cast<UAnimGraphNode_BlendSpaceEvaluator>(*NodeIter);
 
 			// see if sequence player
-			if ( OldNode && OldNode->Node.GetBlendSpace() )
+			if ( OldNode && OldNode->Node.BlendSpace )
 			{
 				// convert to sequence player
 				UEdGraph* TargetGraph = OldNode->GetGraph();
@@ -919,7 +919,7 @@ void FAnimationBlueprintEditor::OnConvertToBlendSpacePlayer()
 				// create new player
 				FGraphNodeCreator<UAnimGraphNode_BlendSpacePlayer> NodeCreator(*TargetGraph);
 				UAnimGraphNode_BlendSpacePlayer* NewNode = NodeCreator.CreateNode();
-				NewNode->Node.SetBlendSpace(OldNode->Node.GetBlendSpace());
+				NewNode->Node.BlendSpace = OldNode->Node.BlendSpace;
 				NodeCreator.Finalize();
 
 				// get default data from old node to new node
@@ -960,7 +960,7 @@ void FAnimationBlueprintEditor::OnConvertToBlendSpaceGraph()
 			UAnimGraphNode_BlendSpacePlayer* OldNode = Cast<UAnimGraphNode_BlendSpacePlayer>(*NodeIter);
 
 			// see if sequence player
-			if (OldNode && OldNode->Node.GetBlendSpace())
+			if (OldNode && OldNode->Node.BlendSpace)
 			{
 				// convert to sequence player
 				UEdGraph* TargetGraph = OldNode->GetGraph();
@@ -970,11 +970,11 @@ void FAnimationBlueprintEditor::OnConvertToBlendSpaceGraph()
 				// create new player
 				FGraphNodeCreator<UAnimGraphNode_BlendSpaceGraph> NodeCreator(*TargetGraph);
 				UAnimGraphNode_BlendSpaceGraph* NewNode = NodeCreator.CreateNode();
-				if(OldNode->Node.GetGroupName() != NAME_None && OldNode->Node.GetGroupMethod() == EAnimSyncMethod::SyncGroup)
+				NewNode->SetupFromAsset(OldNode->Node.BlendSpace, false);
+				if(OldNode->SyncGroup.GroupName != NAME_None && OldNode->SyncGroup.Method == EAnimSyncMethod::SyncGroup)
 				{
-					NewNode->SetSyncGroupName(OldNode->Node.GetGroupName());
+					NewNode->SetSyncGroupName(OldNode->SyncGroup.GroupName);
 				}
-				NewNode->SetupFromAsset(OldNode->Node.GetBlendSpace(), false);
 				NodeCreator.Finalize();
 
 				// get default data from old node to new node
@@ -1116,7 +1116,7 @@ void FAnimationBlueprintEditor::OnConvertToAimOffsetLookAt()
 			UAnimGraphNode_RotationOffsetBlendSpace* OldNode = Cast<UAnimGraphNode_RotationOffsetBlendSpace>(*NodeIter);
 
 			// see if sequence player
-			if (OldNode && OldNode->Node.GetBlendSpace())
+			if (OldNode && OldNode->Node.BlendSpace)
 			{
 				// convert to sequence evaluator
 				UEdGraph* TargetGraph = OldNode->GetGraph();
@@ -1126,7 +1126,7 @@ void FAnimationBlueprintEditor::OnConvertToAimOffsetLookAt()
 				// create new evaluator
 				FGraphNodeCreator<UAnimGraphNode_AimOffsetLookAt> NodeCreator(*TargetGraph);
 				UAnimGraphNode_AimOffsetLookAt* NewNode = NodeCreator.CreateNode();
-				NewNode->Node.SetBlendSpace(OldNode->Node.GetBlendSpace());
+				NewNode->Node.BlendSpace = OldNode->Node.BlendSpace;
 				NodeCreator.Finalize();
 
 				// get default data from old node to new node
@@ -1170,7 +1170,7 @@ void FAnimationBlueprintEditor::OnConvertToAimOffsetSimple()
 			UAnimGraphNode_AimOffsetLookAt* OldNode = Cast<UAnimGraphNode_AimOffsetLookAt>(*NodeIter);
 
 			// see if sequence player
-			if (OldNode && OldNode->Node.GetBlendSpace())
+			if (OldNode && OldNode->Node.BlendSpace)
 			{
 				// convert to sequence player
 				UEdGraph* TargetGraph = OldNode->GetGraph();
@@ -1180,7 +1180,7 @@ void FAnimationBlueprintEditor::OnConvertToAimOffsetSimple()
 				// create new player
 				FGraphNodeCreator<UAnimGraphNode_RotationOffsetBlendSpace> NodeCreator(*TargetGraph);
 				UAnimGraphNode_RotationOffsetBlendSpace* NewNode = NodeCreator.CreateNode();
-				NewNode->Node.SetBlendSpace(OldNode->Node.GetBlendSpace());
+				NewNode->Node.BlendSpace = OldNode->Node.BlendSpace;
 				NodeCreator.Finalize();
 
 				// get default data from old node to new node
@@ -1222,7 +1222,7 @@ void FAnimationBlueprintEditor::OnConvertToAimOffsetGraph()
 			UAnimGraphNode_RotationOffsetBlendSpace* OldNode = Cast<UAnimGraphNode_RotationOffsetBlendSpace>(*NodeIter);
 
 			// see if sequence player
-			if (OldNode && OldNode->Node.GetBlendSpace())
+			if (OldNode && OldNode->Node.BlendSpace)
 			{
 				// convert to sequence player
 				UEdGraph* TargetGraph = OldNode->GetGraph();
@@ -1232,11 +1232,11 @@ void FAnimationBlueprintEditor::OnConvertToAimOffsetGraph()
 				// create new player
 				FGraphNodeCreator<UAnimGraphNode_RotationOffsetBlendSpaceGraph> NodeCreator(*TargetGraph);
 				UAnimGraphNode_RotationOffsetBlendSpaceGraph* NewNode = NodeCreator.CreateNode();
-				if(OldNode->Node.GetGroupName() != NAME_None && OldNode->Node.GetGroupMethod() == EAnimSyncMethod::SyncGroup)
+				NewNode->SetupFromAsset(OldNode->Node.BlendSpace, false);
+				if(OldNode->SyncGroup.GroupName != NAME_None && OldNode->SyncGroup.Method == EAnimSyncMethod::SyncGroup)
 				{
-					NewNode->SetSyncGroupName(OldNode->Node.GetGroupName());
+					NewNode->SetSyncGroupName(OldNode->SyncGroup.GroupName);
 				}
-				NewNode->SetupFromAsset(OldNode->Node.GetBlendSpace(), false);
 				NodeCreator.Finalize();
 
 				// get default data from old node to new node
diff --git a/Engine/Source/Editor/AnimationBlueprintEditor/Private/AnimationNodes/SGraphNodeBlendSpacePlayer.cpp b/Engine/Source/Editor/AnimationBlueprintEditor/Private/AnimationNodes/SGraphNodeBlendSpacePlayer.cpp
--- a/Engine/Source/Editor/AnimationBlueprintEditor/Private/AnimationNodes/SGraphNodeBlendSpacePlayer.cpp
+++ b/Engine/Source/Editor/AnimationBlueprintEditor/Private/AnimationNodes/SGraphNodeBlendSpacePlayer.cpp
@@ -48,7 +48,7 @@ void SGraphNodeBlendSpacePlayer::UpdateGraphSyncLabel()
 			{
 				if (UObject* ActiveObject = AnimBlueprint->GetObjectBeingDebugged())
 				{
-					if(VisualBlendSpacePlayer->Node.GetGroupMethod() == EAnimSyncMethod::Graph)
+					if(VisualBlendSpacePlayer->SyncGroup.Method == EAnimSyncMethod::Graph)
 					{
 						int32 NodeIndex = GeneratedClass->GetNodeIndexFromGuid(VisualBlendSpacePlayer->NodeGuid);
 						if(NodeIndex != INDEX_NONE)
diff --git a/Engine/Source/Editor/AnimationBlueprintEditor/Private/AnimationNodes/SGraphNodeSequencePlayer.cpp b/Engine/Source/Editor/AnimationBlueprintEditor/Private/AnimationNodes/SGraphNodeSequencePlayer.cpp
--- a/Engine/Source/Editor/AnimationBlueprintEditor/Private/AnimationNodes/SGraphNodeSequencePlayer.cpp
+++ b/Engine/Source/Editor/AnimationBlueprintEditor/Private/AnimationNodes/SGraphNodeSequencePlayer.cpp
@@ -176,10 +176,9 @@ void SGraphNodeSequencePlayer::SetSequencePositionRatio(float NewRatio)
 {
 	if(FAnimNode_SequencePlayer* SequencePlayer = GetSequencePlayer())
 	{
-		UAnimSequenceBase* Sequence = SequencePlayer->GetSequence();
-		if (Sequence != NULL)
+		if (SequencePlayer->Sequence != NULL)
 		{
-			const float NewTime = NewRatio * Sequence->GetPlayLength();
+			const float NewTime = NewRatio * SequencePlayer->Sequence->GetPlayLength();
 			SequencePlayer->SetAccumulatedTime(NewTime);
 		}
 	}
@@ -197,7 +196,7 @@ void SGraphNodeSequencePlayer::UpdateGraphSyncLabel()
 			{
 				if (UObject* ActiveObject = AnimBlueprint->GetObjectBeingDebugged())
 				{
-					if(VisualSequencePlayer->Node.GetGroupMethod() == EAnimSyncMethod::Graph)
+					if(VisualSequencePlayer->SyncGroup.Method == EAnimSyncMethod::Graph)
 					{
 						int32 NodeIndex = GeneratedClass->GetNodeIndexFromGuid(VisualSequencePlayer->NodeGuid);
 						if(NodeIndex != INDEX_NONE)
diff --git a/Engine/Source/Editor/UnrealEd/Private/Kismet2/BlueprintEditorUtils.cpp b/Engine/Source/Editor/UnrealEd/Private/Kismet2/BlueprintEditorUtils.cpp
--- a/Engine/Source/Editor/UnrealEd/Private/Kismet2/BlueprintEditorUtils.cpp
+++ b/Engine/Source/Editor/UnrealEd/Private/Kismet2/BlueprintEditorUtils.cpp
@@ -4224,7 +4224,7 @@ void FBlueprintEditorUtils::SetAnimationGraphLayerGroup(UEdGraph* InGraph, const
 	{
 		const FName NewGroup = InGroupName.IsEmpty() ? NAME_None : FName(*InGroupName.ToString());
 		UAnimGraphNode_Root* Root = GetAnimGraphRoot(InGraph);
-		if(NewGroup != Root->Node.GetGroup())
+		if(NewGroup != Root->Node.Group)
 		{
 			FScopedTransaction Transaction(LOCTEXT("SetAnimationGraphLayerGroup", "Set Group"));
 
@@ -4243,7 +4243,7 @@ void FBlueprintEditorUtils::SetAnimationGraphLayerGroup(UEdGraph* InGraph, const
 			}
 
 		
-			Root->Node.SetGroup(NewGroup);
+			Root->Node.Group = NewGroup;
 
 			if (Function)
 			{
diff --git a/Engine/Source/Editor/UnrealEd/Public/IPropertyAccessCompiler.h b/Engine/Source/Editor/UnrealEd/Public/IPropertyAccessCompiler.h
--- a/Engine/Source/Editor/UnrealEd/Public/IPropertyAccessCompiler.h
+++ b/Engine/Source/Editor/UnrealEd/Public/IPropertyAccessCompiler.h
@@ -5,7 +5,6 @@
 #include "CoreMinimal.h"
 #include "Containers/ArrayView.h"
 #include "UObject/ObjectMacros.h"
-#include "Features/IModularFeature.h"
 
 class UObject;
 struct FPropertyAccessLibrary;
@@ -28,10 +27,7 @@ public:
 	virtual ~IPropertyAccessLibraryCompiler() {}
 
 	// Begin compilation - reset the library to its default state
-	virtual void BeginCompilation() = 0;
-	
-	UE_DEPRECATED(5.0, "Please use BeginCompilation without a class argument")
-	virtual void BeginCompilation(const UClass* InClass) {}
+	virtual void BeginCompilation(const UClass* InClass) = 0;
 
 	// Add a copy to the property access library we are compiling
 	// @return an integer handle to the pending copy. This can be resolved to a true copy index by calling MapCopyIndex
diff --git a/Engine/Source/Editor/UnrealEd/Public/IPropertyAccessEditor.h b/Engine/Source/Editor/UnrealEd/Public/IPropertyAccessEditor.h
--- a/Engine/Source/Editor/UnrealEd/Public/IPropertyAccessEditor.h
+++ b/Engine/Source/Editor/UnrealEd/Public/IPropertyAccessEditor.h
@@ -16,8 +16,6 @@ class FExtender;
 class SWidget;
 struct FSlateBrush;
 struct FEdGraphPinType;
-class IPropertyAccessLibraryCompiler;
-struct FPropertyAccessLibrary;
 
 /** An element in a binding chain */
 struct FBindingChainElement
@@ -208,21 +206,6 @@ enum class EPropertyAccessCompatibility
 	Promotable,
 };
 
-// Context used to create a property access library compiler
-struct FPropertyAccessLibraryCompilerArgs
-{
-	FPropertyAccessLibraryCompilerArgs(FPropertyAccessLibrary& InLibrary, const UClass* InClassContext)
-		: Library(InLibrary)
-		, ClassContext(InClassContext)
-	{}
-
-	// The library that will be built
-	FPropertyAccessLibrary& Library;
-
-	// The class that provides a root context for the library to be built in
-	const UClass* ClassContext;
-};
-
 /** Editor support for property access system */
 class IPropertyAccessEditor : public IModularFeature
 {
@@ -256,7 +239,4 @@ public:
 
 	// Makes a string path from a binding chain
 	virtual void MakeStringPath(const TArray<FBindingChainElement>& InBindingChain, TArray<FString>& OutStringPath) const = 0;
-
-	// Make a property access library compiler, used for building a FPropertyAccessLibrary
-	virtual TUniquePtr<IPropertyAccessLibraryCompiler> MakePropertyAccessCompiler(const FPropertyAccessLibraryCompilerArgs& InArgs) const = 0;
 };
\ No newline at end of file
diff --git a/Engine/Source/Runtime/AnimGraphRuntime/Private/AnimGraphRuntimeTrace.cpp b/Engine/Source/Runtime/AnimGraphRuntime/Private/AnimGraphRuntimeTrace.cpp
--- a/Engine/Source/Runtime/AnimGraphRuntime/Private/AnimGraphRuntimeTrace.cpp
+++ b/Engine/Source/Runtime/AnimGraphRuntime/Private/AnimGraphRuntimeTrace.cpp
@@ -34,22 +34,19 @@ void FAnimGraphRuntimeTrace::OutputBlendSpacePlayer(const FAnimationBaseContext&
 	check(InContext.AnimInstanceProxy);
 
 	TRACE_OBJECT(InContext.AnimInstanceProxy->GetAnimInstanceObject());
-	TRACE_OBJECT(InNode.GetBlendSpace());
-
-	FVector SampleCoordinates = InNode.GetPosition();
-	FVector FilteredPosition = InNode.GetFilteredPosition();
+	TRACE_OBJECT(InNode.BlendSpace);
 
 	UE_TRACE_LOG(Animation, BlendSpacePlayer, AnimationChannel)
 		<< BlendSpacePlayer.Cycle(FPlatformTime::Cycles64())
 		<< BlendSpacePlayer.AnimInstanceId(FObjectTrace::GetObjectId(InContext.AnimInstanceProxy->GetAnimInstanceObject()))
-		<< BlendSpacePlayer.BlendSpaceId(FObjectTrace::GetObjectId(InNode.GetBlendSpace()))
+		<< BlendSpacePlayer.BlendSpaceId(FObjectTrace::GetObjectId(InNode.BlendSpace))
 		<< BlendSpacePlayer.NodeId(InContext.GetCurrentNodeId())
-		<< BlendSpacePlayer.PositionX(SampleCoordinates.X)
-		<< BlendSpacePlayer.PositionY(SampleCoordinates.Y)
-		<< BlendSpacePlayer.PositionZ(SampleCoordinates.Z)
-		<< BlendSpacePlayer.FilteredPositionX(FilteredPosition.X)
-		<< BlendSpacePlayer.FilteredPositionY(FilteredPosition.Y)
-		<< BlendSpacePlayer.FilteredPositionZ(FilteredPosition.Z);
+		<< BlendSpacePlayer.PositionX(InNode.X)
+		<< BlendSpacePlayer.PositionY(InNode.Y)
+		<< BlendSpacePlayer.PositionZ(InNode.Z)
+		<< BlendSpacePlayer.FilteredPositionX(InNode.GetFilteredPosition().X)
+		<< BlendSpacePlayer.FilteredPositionY(InNode.GetFilteredPosition().Y)
+		<< BlendSpacePlayer.FilteredPositionZ(InNode.GetFilteredPosition().Z);
 }
 
 void FAnimGraphRuntimeTrace::OutputBlendSpace(const FAnimationBaseContext& InContext, const FAnimNode_BlendSpaceGraphBase& InNode)
diff --git a/Engine/Source/Runtime/AnimGraphRuntime/Private/AnimNodes/AnimNode_AimOffsetLookAt.cpp b/Engine/Source/Runtime/AnimGraphRuntime/Private/AnimNodes/AnimNode_AimOffsetLookAt.cpp
--- a/Engine/Source/Runtime/AnimGraphRuntime/Private/AnimNodes/AnimNode_AimOffsetLookAt.cpp
+++ b/Engine/Source/Runtime/AnimGraphRuntime/Private/AnimNodes/AnimNode_AimOffsetLookAt.cpp
@@ -103,10 +103,10 @@ void FAnimNode_AimOffsetLookAt::Evaluate_AnyThread(FPoseContext& Context)
 
 void FAnimNode_AimOffsetLookAt::UpdateFromLookAtTarget(FPoseContext& LocalPoseContext)
 {
-	UBlendSpace* CurrentBlendSpace = GetBlendSpace();
+	FVector BlendInput(X, Y, Z);
 
 	const FBoneContainer& RequiredBones = LocalPoseContext.Pose.GetBoneContainer();
-	if (CurrentBlendSpace && SocketBoneReference.IsValidToEvaluate(RequiredBones))
+	if (BlendSpace && SocketBoneReference.IsValidToEvaluate(RequiredBones))
 	{
 		FCSPose<FCompactPose> GlobalPose;
 		GlobalPose.InitPose(LocalPoseContext.Pose);
@@ -140,8 +140,8 @@ void FAnimNode_AimOffsetLookAt::UpdateFromLookAtTarget(FPoseContext& LocalPoseCo
 
 		const FVector2D CurrentCoords = FMath::GetAzimuthAndElevation(CurrentDirection, AxisX, AxisY, AxisZ);
 		const FVector2D TargetCoords = FMath::GetAzimuthAndElevation(DirectionToTarget, AxisX, AxisY, AxisZ);
-		CurrentBlendInput.X = FRotator::NormalizeAxis(FMath::RadiansToDegrees(TargetCoords.X - CurrentCoords.X));
-		CurrentBlendInput.Y = FRotator::NormalizeAxis(FMath::RadiansToDegrees(TargetCoords.Y - CurrentCoords.Y));
+		BlendInput.X = FRotator::NormalizeAxis(FMath::RadiansToDegrees(TargetCoords.X - CurrentCoords.X));
+		BlendInput.Y = FRotator::NormalizeAxis(FMath::RadiansToDegrees(TargetCoords.Y - CurrentCoords.Y));
 
 #if ENABLE_DRAW_DEBUG
 		if (CVarAimOffsetLookAtDebug.GetValueOnAnyThread() == 1)
@@ -155,17 +155,21 @@ void FAnimNode_AimOffsetLookAt::UpdateFromLookAtTarget(FPoseContext& LocalPoseCo
 				, FMath::RadiansToDegrees(CurrentCoords.Y)
 				, FMath::RadiansToDegrees(TargetCoords.X)
 				, FMath::RadiansToDegrees(TargetCoords.Y)
-				, CurrentBlendInput.X
-				, CurrentBlendInput.Y);
+				, BlendInput.X
+				, BlendInput.Y);
 			AnimProxy->AnimDrawDebugOnScreenMessage(DebugString, FColor::Red);
 		}
 #endif // ENABLE_DRAW_DEBUG
 	}
 
+	// Set X and Y, so ticking next frame is based on correct weights.
+	X = BlendInput.X;
+	Y = BlendInput.Y;
+
 	// Generate BlendSampleDataCache from inputs.
-	if (CurrentBlendSpace)
+	if (BlendSpace)
 	{
-		CurrentBlendSpace->GetSamplesFromBlendInput(CurrentBlendInput, BlendSampleDataCache, CachedTriangulationIndex, true);
+		BlendSpace->GetSamplesFromBlendInput(BlendInput, BlendSampleDataCache, CachedTriangulationIndex, true);
 	}
 }
 
@@ -180,12 +184,6 @@ void FAnimNode_AimOffsetLookAt::GatherDebugData(FNodeDebugData& DebugData)
 	BasePose.GatherDebugData(DebugData);
 }
 
-FVector FAnimNode_AimOffsetLookAt::GetPosition() const
-{
-	// Use our calculated coordinates rather than the folded values
-	return CurrentBlendInput;
-}
-
 FAnimNode_AimOffsetLookAt::FAnimNode_AimOffsetLookAt()
 	: SocketLocalTransform(FTransform::Identity)
 	, PivotSocketLocalTransform(FTransform::Identity)
@@ -195,7 +193,6 @@ FAnimNode_AimOffsetLookAt::FAnimNode_AimOffsetLookAt()
 	, LookAtLocation(ForceInitToZero)
 	, SocketAxis(1.0f, 0.0f, 0.0f)
 	, Alpha(1.f)
-	, CurrentBlendInput(FVector::ZeroVector)
 	, bIsLODEnabled(false)
 {
 }
diff --git a/Engine/Source/Runtime/AnimGraphRuntime/Private/AnimNodes/AnimNode_BlendListBase.cpp b/Engine/Source/Runtime/AnimGraphRuntime/Private/AnimNodes/AnimNode_BlendListBase.cpp
--- a/Engine/Source/Runtime/AnimGraphRuntime/Private/AnimNodes/AnimNode_BlendListBase.cpp
+++ b/Engine/Source/Runtime/AnimGraphRuntime/Private/AnimNodes/AnimNode_BlendListBase.cpp
@@ -16,8 +16,7 @@ void FAnimNode_BlendListBase::Initialize_AnyThread(const FAnimationInitializeCon
 	FAnimNode_Base::Initialize_AnyThread(Context);
 
 	const int32 NumPoses = BlendPose.Num();
-	const TArray<float>& CurrentBlendTimes = GetBlendTimes();
-	checkSlow(CurrentBlendTimes.Num() == NumPoses);
+	checkSlow(BlendTime.Num() == NumPoses);
 
 	BlendWeights.Reset(NumPoses);
 	PosesToEvaluate.Reset(NumPoses);
@@ -41,8 +40,7 @@ void FAnimNode_BlendListBase::Initialize_AnyThread(const FAnimationInitializeCon
 	Blends.Empty(NumPoses);
 	Blends.AddZeroed(NumPoses);
 
-	UBlendProfile* CurrentBlendProfile = GetBlendProfile();
-	if (CurrentBlendProfile)
+	if (BlendProfile)
 	{
 		BlendStartAlphas.Empty(NumPoses);
 		BlendStartAlphas.AddZeroed(NumPoses);
@@ -50,24 +48,22 @@ void FAnimNode_BlendListBase::Initialize_AnyThread(const FAnimationInitializeCon
 
 	LastActiveChildIndex = INDEX_NONE;
 
-	EAlphaBlendOption CurrentBlendType = GetBlendType();
-	UCurveFloat* CurrentCustomBlendCurve = GetCustomBlendCurve();
 	for(int32 i = 0 ; i < Blends.Num() ; ++i)
 	{
 		FAlphaBlend& Blend = Blends[i];
 
 		Blend.SetBlendTime(0.0f);
-		Blend.SetBlendOption(CurrentBlendType);
-		Blend.SetCustomCurve(CurrentCustomBlendCurve);
+		Blend.SetBlendOption(BlendType);
+		Blend.SetCustomCurve(CustomBlendCurve);
 
-		if (CurrentBlendProfile)
+		if (BlendProfile)
 		{
 			BlendStartAlphas[i] = 0.0f;
 		}
 	}
 	Blends[0].SetAlpha(1.0f);
 
-	if(CurrentBlendProfile)
+	if(BlendProfile)
 	{
 		BlendStartAlphas[0] = 1.0f;
 
@@ -79,7 +75,7 @@ void FAnimNode_BlendListBase::Initialize_AnyThread(const FAnimationInitializeCon
 		{
 			FBlendSampleData& SampleData = PerBoneSampleData[Idx];
 			SampleData.SampleDataIndex = Idx;
-			SampleData.PerBoneBlendData.AddZeroed(CurrentBlendProfile->GetNumBlendEntries());
+			SampleData.PerBoneBlendData.AddZeroed(BlendProfile->GetNumBlendEntries());
 		}
 	}
 }
@@ -99,15 +95,12 @@ void FAnimNode_BlendListBase::Update_AnyThread(const FAnimationUpdateContext& Co
 	GetEvaluateGraphExposedInputs().Execute(Context);
 
 	const int32 NumPoses = BlendPose.Num();
-	const TArray<float>& CurrentBlendTimes = GetBlendTimes();
-	checkSlow((CurrentBlendTimes.Num() == NumPoses) && (BlendWeights.Num() == NumPoses));
+	checkSlow((BlendTime.Num() == NumPoses) && (BlendWeights.Num() == NumPoses));
 
 	PosesToEvaluate.Empty(NumPoses);
 
 	if (NumPoses > 0)
 	{
-		UBlendProfile* CurrentBlendProfile = GetBlendProfile();
-		
 		// Handle a change in the active child index; adjusting the target weights
 		const int32 ChildIndex = GetActiveChildIndex();
 		
@@ -127,12 +120,12 @@ void FAnimNode_BlendListBase::Update_AnyThread(const FAnimationUpdateContext& Co
 			{
 				RemainingBlendTime = 0.0f;
 			}
-			else if (GetTransitionType() == EBlendListTransitionType::Inertialization)
+			else if (TransitionType == EBlendListTransitionType::Inertialization)
 			{
 				UE::Anim::IInertializationRequester* InertializationRequester = Context.GetMessage<UE::Anim::IInertializationRequester>();
 				if (InertializationRequester)
 				{
-					InertializationRequester->RequestInertialization(CurrentBlendTimes[ChildIndex]);
+					InertializationRequester->RequestInertialization(BlendTime[ChildIndex]);
 					InertializationRequester->AddDebugRecord(*Context.AnimInstanceProxy, Context.GetCurrentNodeId());
 				}
 				else
@@ -144,7 +137,7 @@ void FAnimNode_BlendListBase::Update_AnyThread(const FAnimationUpdateContext& Co
 			}
 			else
 			{
-				RemainingBlendTime = CurrentBlendTimes[ChildIndex] * WeightDifference;
+				RemainingBlendTime = BlendTime[ChildIndex] * WeightDifference;
 			}
 
 			for (int32 i = 0; i < RemainingBlendTimes.Num(); ++i)
@@ -168,7 +161,7 @@ void FAnimNode_BlendListBase::Update_AnyThread(const FAnimationUpdateContext& Co
 				{
 					Blend.SetValueRange(BlendWeights[i], 1.0f);
 
-					if (CurrentBlendProfile)
+					if (BlendProfile)
 					{
 						Blend.ResetAlpha();
 						BlendStartAlphas[i] = Blend.GetAlpha();
@@ -179,14 +172,14 @@ void FAnimNode_BlendListBase::Update_AnyThread(const FAnimationUpdateContext& Co
 					Blend.SetValueRange(BlendWeights[i], 0.0f);
 				}
 
-				if (CurrentBlendProfile)
+				if (BlendProfile)
 				{
 					BlendStartAlphas[i] = Blend.GetAlpha();
 				}
 			}
 
 			// when this flag is true, we'll reinitialize the children
-			if (GetResetChildOnActivation())
+			if (bResetChildOnActivation)
 			{
 				FAnimationInitializeContext ReinitializeContext(Context.AnimInstanceProxy, Context.SharedContext);
 
@@ -233,13 +226,13 @@ void FAnimNode_BlendListBase::Update_AnyThread(const FAnimationUpdateContext& Co
 		}
 
 		// If we're using a blend profile, extract the scales and build blend sample data
-		if (CurrentBlendProfile)
+		if (BlendProfile)
 		{
 			for(int32 i = 0; i < BlendPose.Num(); ++i)
 			{
 				FBlendSampleData& PoseSampleData = PerBoneSampleData[i];
 				PoseSampleData.TotalWeight = BlendWeights[i];
-				CurrentBlendProfile->UpdateBoneWeights(PoseSampleData, Blends[i], BlendStartAlphas[i], BlendWeights[i]);
+				BlendProfile->UpdateBoneWeights(PoseSampleData, Blends[i], BlendStartAlphas[i], BlendWeights[i]);
 			}
 
 			FBlendSampleData::NormalizeDataWeight(PerBoneSampleData);
@@ -250,7 +243,7 @@ void FAnimNode_BlendListBase::Update_AnyThread(const FAnimationUpdateContext& Co
 	const int32 ChildIndex = GetActiveChildIndex();
 	TRACE_ANIM_NODE_VALUE(Context, TEXT("Active Index"), ChildIndex);
 	TRACE_ANIM_NODE_VALUE(Context, TEXT("Active Weight"), BlendWeights[ChildIndex]);
-	TRACE_ANIM_NODE_VALUE(Context, TEXT("Active Blend Time"), CurrentBlendTimes[ChildIndex]);
+	TRACE_ANIM_NODE_VALUE(Context, TEXT("Active Blend Time"), BlendTime[ChildIndex]);
 #endif
 }
 
@@ -290,10 +283,9 @@ void FAnimNode_BlendListBase::Evaluate_AnyThread(FPoseContext& Output)
 		FAnimationPoseData OutAnimationPoseData(Output);
 
 		// Use the calculated blend sample data if we're blending per-bone
-		UBlendProfile* CurrentBlendProfile = GetBlendProfile();
-		if (CurrentBlendProfile)
+		if (BlendProfile)
 		{
-			FAnimationRuntime::BlendPosesTogetherPerBone(FilteredPoses, FilteredCurve, FilteredAttributes, CurrentBlendProfile, PerBoneSampleData, PosesToEvaluate, OutAnimationPoseData);
+			FAnimationRuntime::BlendPosesTogetherPerBone(FilteredPoses, FilteredCurve, FilteredAttributes, BlendProfile, PerBoneSampleData, PosesToEvaluate, OutAnimationPoseData);
 		}
 		else
 		{
@@ -313,8 +305,7 @@ void FAnimNode_BlendListBase::GatherDebugData(FNodeDebugData& DebugData)
 	const int32 ChildIndex = GetActiveChildIndex();
 
 	FString DebugLine = GetNodeName(DebugData);
-	const TArray<float>& CurrentBlendTimes = GetBlendTimes();
-	DebugLine += FString::Printf(TEXT("(Active: (%i/%i) Weight: %.1f%% Time %.3f)"), ChildIndex+1, NumPoses, BlendWeights[ChildIndex]*100.f, CurrentBlendTimes[ChildIndex]);
+	DebugLine += FString::Printf(TEXT("(Active: (%i/%i) Weight: %.1f%% Time %.3f)"), ChildIndex+1, NumPoses, BlendWeights[ChildIndex]*100.f, BlendTime[ChildIndex]);
 
 	DebugData.AddDebugItem(DebugLine);
 	
@@ -323,33 +314,3 @@ void FAnimNode_BlendListBase::GatherDebugData(FNodeDebugData& DebugData)
 		BlendPose[Pose].GatherDebugData(DebugData.BranchFlow(BlendWeights[Pose]));
 	}
 }
-
-const TArray<float>& FAnimNode_BlendListBase::GetBlendTimes() const
-{
-	return GET_ANIM_NODE_DATA(TArray<float>, BlendTime);
-}
-
-EBlendListTransitionType FAnimNode_BlendListBase::GetTransitionType() const
-{
-	return GET_ANIM_NODE_DATA(EBlendListTransitionType, TransitionType);
-}
-
-EAlphaBlendOption FAnimNode_BlendListBase::GetBlendType() const
-{
-	return GET_ANIM_NODE_DATA(EAlphaBlendOption, BlendType);
-}
-
-bool FAnimNode_BlendListBase::GetResetChildOnActivation() const
-{
-	return GET_ANIM_NODE_DATA(bool, bResetChildOnActivation);
-}
-
-UCurveFloat* FAnimNode_BlendListBase::GetCustomBlendCurve() const
-{
-	return GET_ANIM_NODE_DATA(TObjectPtr<UCurveFloat>, CustomBlendCurve);
-}
-
-UBlendProfile* FAnimNode_BlendListBase::GetBlendProfile() const
-{
-	return GET_ANIM_NODE_DATA(TObjectPtr<UBlendProfile>, BlendProfile);
-}
\ No newline at end of file
diff --git a/Engine/Source/Runtime/AnimGraphRuntime/Private/AnimNodes/AnimNode_BlendListByBool.cpp b/Engine/Source/Runtime/AnimGraphRuntime/Private/AnimNodes/AnimNode_BlendListByBool.cpp
--- a/Engine/Source/Runtime/AnimGraphRuntime/Private/AnimNodes/AnimNode_BlendListByBool.cpp
+++ b/Engine/Source/Runtime/AnimGraphRuntime/Private/AnimNodes/AnimNode_BlendListByBool.cpp
@@ -8,10 +8,5 @@
 int32 FAnimNode_BlendListByBool::GetActiveChildIndex()
 {
 	// Note: Intentionally flipped boolean sense (the true input is #0, and the false input is #1)
-	return GetActiveValue() ? 0 : 1;
+	return bActiveValue ? 0 : 1;
 }
-
-bool FAnimNode_BlendListByBool::GetActiveValue() const
-{
-	return GET_ANIM_NODE_DATA(bool, bActiveValue);
-}
\ No newline at end of file
diff --git a/Engine/Source/Runtime/AnimGraphRuntime/Private/AnimNodes/AnimNode_BlendListByEnum.cpp b/Engine/Source/Runtime/AnimGraphRuntime/Private/AnimNodes/AnimNode_BlendListByEnum.cpp
--- a/Engine/Source/Runtime/AnimGraphRuntime/Private/AnimNodes/AnimNode_BlendListByEnum.cpp
+++ b/Engine/Source/Runtime/AnimGraphRuntime/Private/AnimNodes/AnimNode_BlendListByEnum.cpp
@@ -7,11 +7,9 @@
 
 int32 FAnimNode_BlendListByEnum::GetActiveChildIndex()
 {
-	uint8 CurrentActiveEnumValue = GetActiveEnumValue();
-	const TArray<int32>& CurrentEnumToPoseIndex = GetEnumToPoseIndex();
-	if (CurrentEnumToPoseIndex.IsValidIndex(CurrentActiveEnumValue))
+	if (EnumToPoseIndex.IsValidIndex(ActiveEnumValue))
 	{
-		return CurrentEnumToPoseIndex[CurrentActiveEnumValue];
+		return EnumToPoseIndex[ActiveEnumValue];
 	}
 	else
 	{
@@ -19,19 +17,3 @@ int32 FAnimNode_BlendListByEnum::GetActiveChildIndex()
 	}
 }
 
-#if WITH_EDITORONLY_DATA
-void FAnimNode_BlendListByEnum::SetEnumToPoseIndex(const TArray<int32>& InEnumToPoseIndex)
-{
-	EnumToPoseIndex = InEnumToPoseIndex;
-}
-#endif
-
-const TArray<int32>& FAnimNode_BlendListByEnum::GetEnumToPoseIndex() const
-{
-	return GET_ANIM_NODE_DATA(TArray<int32>, EnumToPoseIndex);
-}
-
-uint8 FAnimNode_BlendListByEnum::GetActiveEnumValue() const
-{
-	return GET_ANIM_NODE_DATA(uint8, ActiveEnumValue);
-}
\ No newline at end of file
diff --git a/Engine/Source/Runtime/AnimGraphRuntime/Private/AnimNodes/AnimNode_BlendListByInt.cpp b/Engine/Source/Runtime/AnimGraphRuntime/Private/AnimNodes/AnimNode_BlendListByInt.cpp
--- a/Engine/Source/Runtime/AnimGraphRuntime/Private/AnimNodes/AnimNode_BlendListByInt.cpp
+++ b/Engine/Source/Runtime/AnimGraphRuntime/Private/AnimNodes/AnimNode_BlendListByInt.cpp
@@ -7,7 +7,6 @@
 
 int32 FAnimNode_BlendListByInt::GetActiveChildIndex()
 {
-	const int32 NumPoses = BlendPose.Num();
-	int32 CurrentActiveChildIndex = GET_ANIM_NODE_DATA(int32, ActiveChildIndex);
-	return FMath::Clamp<int32>(CurrentActiveChildIndex, 0, NumPoses - 1);
+	const int NumPoses = BlendPose.Num();
+	return FMath::Clamp<int32>(ActiveChildIndex, 0, NumPoses - 1);
 }
diff --git a/Engine/Source/Runtime/AnimGraphRuntime/Private/AnimNodes/AnimNode_BlendSpacePlayer.cpp b/Engine/Source/Runtime/AnimGraphRuntime/Private/AnimNodes/AnimNode_BlendSpacePlayer.cpp
--- a/Engine/Source/Runtime/AnimGraphRuntime/Private/AnimNodes/AnimNode_BlendSpacePlayer.cpp
+++ b/Engine/Source/Runtime/AnimGraphRuntime/Private/AnimNodes/AnimNode_BlendSpacePlayer.cpp
@@ -11,7 +11,20 @@
 /////////////////////////////////////////////////////
 // FAnimNode_BlendSpacePlayer
 
-float FAnimNode_BlendSpacePlayer::GetCurrentAssetTime() const
+FAnimNode_BlendSpacePlayer::FAnimNode_BlendSpacePlayer()
+	: X(0.0f)
+	, Y(0.0f)
+	, Z(0.0f)
+	, PlayRate(1.0f)
+	, bLoop(true)
+	, bResetPlayTimeWhenBlendSpaceChanges(true)
+	, StartPosition(0.f)
+	, BlendSpace(nullptr)
+	, PreviousBlendSpace(nullptr)
+{
+}
+
+float FAnimNode_BlendSpacePlayer::GetCurrentAssetTime()
 {
 	if(const FBlendSampleData* HighestWeightedSample = GetHighestWeightedSample())
 	{
@@ -22,20 +35,19 @@ float FAnimNode_BlendSpacePlayer::GetCurrentAssetTime() const
 	return 0.0f;
 }
 
-float FAnimNode_BlendSpacePlayer::GetCurrentAssetTimePlayRateAdjusted() const
+float FAnimNode_BlendSpacePlayer::GetCurrentAssetTimePlayRateAdjusted()
 {
 	float Length = GetCurrentAssetLength();
-	return GetPlayRate() < 0.0f ? Length - InternalTimeAccumulator * Length : Length * InternalTimeAccumulator;
+	return PlayRate < 0.0f ? Length - InternalTimeAccumulator * Length : Length * InternalTimeAccumulator;
 }
 
-float FAnimNode_BlendSpacePlayer::GetCurrentAssetLength() const
+float FAnimNode_BlendSpacePlayer::GetCurrentAssetLength()
 {
 	if(const FBlendSampleData* HighestWeightedSample = GetHighestWeightedSample())
 	{
-		UBlendSpace* CurrentBlendSpace = GetBlendSpace();
-		if (CurrentBlendSpace != nullptr)
+		if (BlendSpace != nullptr)
 		{
-			const FBlendSample& Sample = CurrentBlendSpace->GetBlendSample(HighestWeightedSample->SampleDataIndex);
+			const FBlendSample& Sample = BlendSpace->GetBlendSample(HighestWeightedSample->SampleDataIndex);
 			return Sample.Animation->GetPlayLength();
 		}
 	}
@@ -53,7 +65,7 @@ void FAnimNode_BlendSpacePlayer::Initialize_AnyThread(const FAnimationInitialize
 
 	Reinitialize();
 
-	PreviousBlendSpace = GetBlendSpace();
+	PreviousBlendSpace = BlendSpace;
 }
 
 void FAnimNode_BlendSpacePlayer::CacheBones_AnyThread(const FAnimationCacheBonesContext& Context)
@@ -71,23 +83,21 @@ void FAnimNode_BlendSpacePlayer::UpdateAssetPlayer(const FAnimationUpdateContext
 void FAnimNode_BlendSpacePlayer::UpdateInternal(const FAnimationUpdateContext& Context)
 {
 	DECLARE_SCOPE_HIERARCHICAL_COUNTER_ANIMNODE(UpdateInternal)
-
-	UBlendSpace* CurrentBlendSpace = GetBlendSpace();
-	if ((CurrentBlendSpace != nullptr) && (Context.AnimInstanceProxy->IsSkeletonCompatible(CurrentBlendSpace->GetSkeleton())))
+	if ((BlendSpace != nullptr) && (Context.AnimInstanceProxy->IsSkeletonCompatible(BlendSpace->GetSkeleton())))
 	{	
-		if (PreviousBlendSpace != CurrentBlendSpace)
+		if (PreviousBlendSpace != BlendSpace)
 		{
-			Reinitialize(ShouldResetPlayTimeWhenBlendSpaceChanges());
+			Reinitialize(bResetPlayTimeWhenBlendSpaceChanges);
 		}
 
-		const FVector Position = GetPosition();
+		const FVector Position(X, Y, Z);
 
 		// Create a tick record and push into the closest scope
 		UE::Anim::FAnimSyncGroupScope& SyncScope = Context.GetMessageChecked<UE::Anim::FAnimSyncGroupScope>();
 
-		FAnimTickRecord TickRecord(CurrentBlendSpace, Position, BlendSampleDataCache, BlendFilter, GetLoop(), GetPlayRate(), Context.GetFinalBlendWeight(), /*inout*/ InternalTimeAccumulator, MarkerTickRecord);
+		FAnimTickRecord TickRecord(BlendSpace, Position, BlendSampleDataCache, BlendFilter, bLoop, PlayRate, Context.GetFinalBlendWeight(), /*inout*/ InternalTimeAccumulator, MarkerTickRecord);
 		TickRecord.RootMotionWeightModifier = Context.GetRootMotionWeightModifier();
-		UE::Anim::FAnimSyncParams SyncParams(GetGroupName(), GetGroupRole(), GetGroupMethod());
+		UE::Anim::FAnimSyncParams SyncParams(GroupName, GroupRole, Method);
 
 		SyncScope.AddTickRecord(TickRecord, SyncParams, UE::Anim::FAnimSyncDebugInfo(Context));
 
@@ -96,28 +106,26 @@ void FAnimNode_BlendSpacePlayer::UpdateInternal(const FAnimationUpdateContext& C
 #if WITH_EDITORONLY_DATA
 		if (FAnimBlueprintDebugData* DebugData = Context.AnimInstanceProxy->GetAnimBlueprintDebugData())
 		{
-			DebugData->RecordBlendSpacePlayer(Context.GetCurrentNodeId(), CurrentBlendSpace, Position, BlendFilter.GetFilterLastOutput());
+			DebugData->RecordBlendSpacePlayer(Context.GetCurrentNodeId(), BlendSpace, Position, BlendFilter.GetFilterLastOutput());
 		}
 #endif
 
-		PreviousBlendSpace = CurrentBlendSpace;
+		PreviousBlendSpace = BlendSpace;
 	}
 
 	TRACE_BLENDSPACE_PLAYER(Context, *this);
-	TRACE_ANIM_NODE_VALUE(Context, TEXT("Name"), CurrentBlendSpace ? *CurrentBlendSpace->GetName() : TEXT("None"));
-	TRACE_ANIM_NODE_VALUE(Context, TEXT("Blend Space"), CurrentBlendSpace);
+	TRACE_ANIM_NODE_VALUE(Context, TEXT("Name"), BlendSpace ? *BlendSpace->GetName() : TEXT("None"));
+	TRACE_ANIM_NODE_VALUE(Context, TEXT("Blend Space"), BlendSpace);
 	TRACE_ANIM_NODE_VALUE(Context, TEXT("Playback Time"), InternalTimeAccumulator);
 }
 
 void FAnimNode_BlendSpacePlayer::Evaluate_AnyThread(FPoseContext& Output)
 {
 	DECLARE_SCOPE_HIERARCHICAL_COUNTER_ANIMNODE(Evaluate_AnyThread)
-
-	UBlendSpace* CurrentBlendSpace = GetBlendSpace();
-	if ((CurrentBlendSpace != nullptr) && (Output.AnimInstanceProxy->IsSkeletonCompatible(CurrentBlendSpace->GetSkeleton())))
+	if ((BlendSpace != nullptr) && (Output.AnimInstanceProxy->IsSkeletonCompatible(BlendSpace->GetSkeleton())))
 	{
 		FAnimationPoseData AnimationPoseData(Output);
-		CurrentBlendSpace->GetAnimationPose(BlendSampleDataCache, AnimationPoseData);
+		BlendSpace->GetAnimationPose(BlendSampleDataCache, AnimationPoseData);
 	}
 	else
 	{
@@ -125,30 +133,38 @@ void FAnimNode_BlendSpacePlayer::Evaluate_AnyThread(FPoseContext& Output)
 	}
 }
 
+void FAnimNode_BlendSpacePlayer::OverrideAsset(UAnimationAsset* NewAsset)
+{
+	if(UBlendSpace* NewBlendSpace = Cast<UBlendSpace>(NewAsset))
+	{
+		BlendSpace = NewBlendSpace;
+	}
+}
+
 void FAnimNode_BlendSpacePlayer::GatherDebugData(FNodeDebugData& DebugData)
 {
 	DECLARE_SCOPE_HIERARCHICAL_COUNTER_ANIMNODE(GatherDebugData)
 	FString DebugLine = DebugData.GetNodeName(this);
-
-	UBlendSpace* CurrentBlendSpace = GetBlendSpace();
-	if (CurrentBlendSpace)
+	if (BlendSpace)
 	{
-		DebugLine += FString::Printf(TEXT("('%s' Play Time: %.3f)"), *CurrentBlendSpace->GetName(), InternalTimeAccumulator);
+		DebugLine += FString::Printf(TEXT("('%s' Play Time: %.3f)"), *BlendSpace->GetName(), InternalTimeAccumulator);
 
 		DebugData.AddDebugItem(DebugLine, true);
+
+		//BlendSpace->GetBlendSamples();
 	}
 }
 
-float FAnimNode_BlendSpacePlayer::GetTimeFromEnd(float CurrentTime) const
+float FAnimNode_BlendSpacePlayer::GetTimeFromEnd(float CurrentTime)
 {
 	// Blend-spaces use normalized time value
 	const float PlayLength = 1.0f;
-	return GetBlendSpace() != nullptr ? PlayLength - CurrentTime : 0.0f;
+	return BlendSpace != nullptr ? PlayLength - CurrentTime : 0.0f;
 }
 
-UAnimationAsset* FAnimNode_BlendSpacePlayer::GetAnimAsset() const
+UAnimationAsset* FAnimNode_BlendSpacePlayer::GetAnimAsset()
 {
-	return GetBlendSpace();
+	return BlendSpace;
 }
 
 const FBlendSampleData* FAnimNode_BlendSpacePlayer::GetHighestWeightedSample() const
@@ -176,57 +192,16 @@ void FAnimNode_BlendSpacePlayer::Reinitialize(bool bResetTime)
 	BlendSampleDataCache.Empty();
 	if(bResetTime)
 	{
-		float CurrentStartPosition = GetStartPosition();
-
-		InternalTimeAccumulator = FMath::Clamp(CurrentStartPosition, 0.f, 1.0f);
-		if (CurrentStartPosition == 0.f && GetPlayRate() < 0.0f)
+		InternalTimeAccumulator = FMath::Clamp(StartPosition, 0.f, 1.0f);
+		if (StartPosition == 0.f && PlayRate < 0.0f)
 		{
 			// Blend spaces run between 0 and 1
 			InternalTimeAccumulator = 1.0f;
 		}
 	}
 
-	UBlendSpace* CurrentBlendSpace = GetBlendSpace();
-	if (CurrentBlendSpace != nullptr)
+	if (BlendSpace != nullptr)
 	{
-		CurrentBlendSpace->InitializeFilter(&BlendFilter);
+		BlendSpace->InitializeFilter(&BlendFilter);
 	}
 }
-
-#if WITH_EDITORONLY_DATA
-void FAnimNode_BlendSpacePlayer::SetBlendSpace(UBlendSpace* InBlendSpace)
-{
-	TObjectPtr<UBlendSpace>& BlendSpaceToSet = GET_MUTABLE_ANIM_NODE_DATA(TObjectPtr<UBlendSpace>, BlendSpace);
-	BlendSpaceToSet = InBlendSpace;
-}
-#endif
-
-FVector FAnimNode_BlendSpacePlayer::GetPosition() const
-{
-	return FVector(GET_ANIM_NODE_DATA(float, X), GET_ANIM_NODE_DATA(float, Y), GET_ANIM_NODE_DATA(float, Z));
-}
-
-float FAnimNode_BlendSpacePlayer::GetPlayRate() const
-{
-	return GET_ANIM_NODE_DATA(float, PlayRate);
-}
-
-bool FAnimNode_BlendSpacePlayer::GetLoop() const
-{
-	return GET_ANIM_NODE_DATA(bool, bLoop);
-}
-
-bool FAnimNode_BlendSpacePlayer::ShouldResetPlayTimeWhenBlendSpaceChanges() const
-{
-	return GET_ANIM_NODE_DATA(bool, bResetPlayTimeWhenBlendSpaceChanges);
-}
-
-float FAnimNode_BlendSpacePlayer::GetStartPosition() const
-{
-	return GET_ANIM_NODE_DATA(float, StartPosition);
-}
-
-UBlendSpace* FAnimNode_BlendSpacePlayer::GetBlendSpace() const
-{
-	return GET_ANIM_NODE_DATA(TObjectPtr<UBlendSpace>, BlendSpace);
-}
diff --git a/Engine/Source/Runtime/AnimGraphRuntime/Private/AnimNodes/AnimNode_PoseHandler.cpp b/Engine/Source/Runtime/AnimGraphRuntime/Private/AnimNodes/AnimNode_PoseHandler.cpp
--- a/Engine/Source/Runtime/AnimGraphRuntime/Private/AnimNodes/AnimNode_PoseHandler.cpp
+++ b/Engine/Source/Runtime/AnimGraphRuntime/Private/AnimNodes/AnimNode_PoseHandler.cpp
@@ -83,12 +83,13 @@ void FAnimNode_PoseHandler::UpdateAssetPlayer(const FAnimationUpdateContext& Con
 	TRACE_ANIM_NODE_VALUE(Context, TEXT("Pose Asset"), CurrentPoseAsset.IsValid() ? *CurrentPoseAsset.Get()->GetName() : TEXT("None"));
 }
 
-#if WITH_EDITORONLY_DATA
-void FAnimNode_PoseHandler::SetPoseAsset(UPoseAsset* InPoseAsset)
+void FAnimNode_PoseHandler::OverrideAsset(UAnimationAsset* NewAsset)
 {
-	PoseAsset = InPoseAsset;
+	if(UPoseAsset* NewPoseAsset = Cast<UPoseAsset>(NewAsset))
+	{
+		PoseAsset = NewPoseAsset;
+	}
 }
-#endif
 
 void FAnimNode_PoseHandler::GatherDebugData(FNodeDebugData& DebugData)
 {
diff --git a/Engine/Source/Runtime/AnimGraphRuntime/Private/AnimNodes/AnimNode_RefPose.cpp b/Engine/Source/Runtime/AnimGraphRuntime/Private/AnimNodes/AnimNode_RefPose.cpp
--- a/Engine/Source/Runtime/AnimGraphRuntime/Private/AnimNodes/AnimNode_RefPose.cpp
+++ b/Engine/Source/Runtime/AnimGraphRuntime/Private/AnimNodes/AnimNode_RefPose.cpp
@@ -26,8 +26,7 @@ void FAnimNode_RefPose::Evaluate_AnyThread(FPoseContext& Output)
 	// I don't have anything to evaluate. Should this be even here?
 	// EvaluateGraphExposedInputs.Execute(Context);
 
-    ERefPoseType CurrentRefPoseType = GetRefPoseType(); 
-	switch (CurrentRefPoseType) 
+	switch (RefPoseType) 
 	{
 	case EIT_LocalSpace:
 		Output.ResetToRefPose();
@@ -39,12 +38,7 @@ void FAnimNode_RefPose::Evaluate_AnyThread(FPoseContext& Output)
 		break;
 	}
 
-	TRACE_ANIM_NODE_VALUE(Output, TEXT("Ref Pose Type"), GetRefPostTypeText(CurrentRefPoseType));
-}
-
-ERefPoseType FAnimNode_RefPose::GetRefPoseType() const
-{
-	return GET_ANIM_NODE_DATA(TEnumAsByte<ERefPoseType>, RefPoseType);
+	TRACE_ANIM_NODE_VALUE(Output, TEXT("Ref Pose Type"), GetRefPostTypeText(RefPoseType));
 }
 
 void FAnimNode_MeshSpaceRefPose::EvaluateComponentSpace_AnyThread(FComponentSpacePoseContext& Output)
@@ -57,6 +51,6 @@ void FAnimNode_RefPose::GatherDebugData(FNodeDebugData& DebugData)
 {
 	DECLARE_SCOPE_HIERARCHICAL_COUNTER_ANIMNODE(GatherDebugData)
 	FString DebugLine = DebugData.GetNodeName(this);
-	DebugLine += FString::Printf(TEXT("(Ref Pose Type: %s)"), GetRefPostTypeText(GetRefPoseType()));
+	DebugLine += FString::Printf(TEXT("(Ref Pose Type: %s)"), GetRefPostTypeText(RefPoseType));
 	DebugData.AddDebugItem(DebugLine, true);
 }
diff --git a/Engine/Source/Runtime/AnimGraphRuntime/Private/AnimNodes/AnimNode_SequenceEvaluator.cpp b/Engine/Source/Runtime/AnimGraphRuntime/Private/AnimNodes/AnimNode_SequenceEvaluator.cpp
--- a/Engine/Source/Runtime/AnimGraphRuntime/Private/AnimNodes/AnimNode_SequenceEvaluator.cpp
+++ b/Engine/Source/Runtime/AnimGraphRuntime/Private/AnimNodes/AnimNode_SequenceEvaluator.cpp
@@ -4,67 +4,65 @@
 #include "Animation/AnimInstanceProxy.h"
 #include "Animation/AnimTrace.h"
 
-float FAnimNode_SequenceEvaluatorBase::GetCurrentAssetTime() const
+float FAnimNode_SequenceEvaluator::GetCurrentAssetTime()
 {
-	return GetExplicitTime();
+	return ExplicitTime;
 }
 
-float FAnimNode_SequenceEvaluatorBase::GetCurrentAssetLength() const
+float FAnimNode_SequenceEvaluator::GetCurrentAssetLength()
 {
-	UAnimSequenceBase* CurrentSequence = GetSequence();
-	return CurrentSequence ? CurrentSequence->GetPlayLength() : 0.0f;
+	return Sequence ? Sequence->GetPlayLength() : 0.0f;
 }
 
-void FAnimNode_SequenceEvaluatorBase::Initialize_AnyThread(const FAnimationInitializeContext& Context)
+/////////////////////////////////////////////////////
+// FAnimSequenceEvaluatorNode
+
+void FAnimNode_SequenceEvaluator::Initialize_AnyThread(const FAnimationInitializeContext& Context)
 {
 	DECLARE_SCOPE_HIERARCHICAL_COUNTER_ANIMNODE(Initialize_AnyThread)
 	FAnimNode_AssetPlayerBase::Initialize_AnyThread(Context);
 	bReinitialized = true;
 }
 
-void FAnimNode_SequenceEvaluatorBase::CacheBones_AnyThread(const FAnimationCacheBonesContext& Context)
+void FAnimNode_SequenceEvaluator::CacheBones_AnyThread(const FAnimationCacheBonesContext& Context)
 {
 	DECLARE_SCOPE_HIERARCHICAL_COUNTER_ANIMNODE(CacheBones_AnyThread)
 }
 
-void FAnimNode_SequenceEvaluatorBase::UpdateAssetPlayer(const FAnimationUpdateContext& Context)
+void FAnimNode_SequenceEvaluator::UpdateAssetPlayer(const FAnimationUpdateContext& Context)
 {
 	GetEvaluateGraphExposedInputs().Execute(Context);
 
-	float CurrentExplicitTime = GetExplicitTime();
-
-	UAnimSequenceBase* CurrentSequence = GetSequence();
-	if (CurrentSequence)
+	if (Sequence)
 	{
 		// Clamp input to a valid position on this sequence's time line.
-		CurrentExplicitTime = FMath::Clamp(CurrentExplicitTime, 0.f, CurrentSequence->GetPlayLength());
+		ExplicitTime = FMath::Clamp(ExplicitTime, 0.f, Sequence->GetPlayLength());
 
-		if ((!GetTeleportToExplicitTime() || (GetGroupName() != NAME_None)) && (Context.AnimInstanceProxy->IsSkeletonCompatible(CurrentSequence->GetSkeleton())))
+		if ((!bTeleportToExplicitTime || (GroupName != NAME_None)) && (Context.AnimInstanceProxy->IsSkeletonCompatible(Sequence->GetSkeleton())))
 		{
 			if (bReinitialized)
 			{
-				switch (GetReinitializationBehavior())
+				switch (ReinitializationBehavior)
 				{
-					case ESequenceEvalReinit::StartPosition: InternalTimeAccumulator = GetStartPosition(); break;
-					case ESequenceEvalReinit::ExplicitTime: InternalTimeAccumulator = CurrentExplicitTime; break;
+					case ESequenceEvalReinit::StartPosition: InternalTimeAccumulator = StartPosition; break;
+					case ESequenceEvalReinit::ExplicitTime: InternalTimeAccumulator = ExplicitTime; break;
 				}
 
-				InternalTimeAccumulator = FMath::Clamp(InternalTimeAccumulator, 0.f, CurrentSequence->GetPlayLength());
+				InternalTimeAccumulator = FMath::Clamp(InternalTimeAccumulator, 0.f, Sequence->GetPlayLength());
 			}
 
-			float TimeJump = CurrentExplicitTime - InternalTimeAccumulator;
-			const bool bCurrentShouldLoop = GetShouldLoop();
-			if (bCurrentShouldLoop)
+			float TimeJump = ExplicitTime - InternalTimeAccumulator;
+			if (bShouldLoop)
 			{
-				if (FMath::Abs(TimeJump) > (CurrentSequence->GetPlayLength() * 0.5f))
+				if (FMath::Abs(TimeJump) > (Sequence->GetPlayLength() * 0.5f))
 				{
 					if (TimeJump > 0.f)
 					{
-						TimeJump -= CurrentSequence->GetPlayLength();
+						TimeJump -= Sequence->GetPlayLength();
 					}
 					else
 					{
-						TimeJump += CurrentSequence->GetPlayLength();
+						TimeJump += Sequence->GetPlayLength();
 					}
 				}
 			}
@@ -73,37 +71,36 @@ void FAnimNode_SequenceEvaluatorBase::UpdateAssetPlayer(const FAnimationUpdateCo
 			// to prevent that from happening, we set current accumulator to explicit time
 			if (TimeJump == 0.f)
 			{
-				InternalTimeAccumulator = CurrentExplicitTime;
+				InternalTimeAccumulator = ExplicitTime;
 			}
 			
 			const float DeltaTime = Context.GetDeltaTime();
-			const float RateScale = CurrentSequence->RateScale;
+			const float RateScale = Sequence->RateScale;
 			const float PlayRate = FMath::IsNearlyZero(DeltaTime) || FMath::IsNearlyZero(RateScale) ? 0.f : (TimeJump / (DeltaTime * RateScale));
-			CreateTickRecordForNode(Context, CurrentSequence, bCurrentShouldLoop, PlayRate);
+			CreateTickRecordForNode(Context, Sequence, bShouldLoop, PlayRate);
 		}
 		else
 		{
-			InternalTimeAccumulator = CurrentExplicitTime;
+			InternalTimeAccumulator = ExplicitTime;
 		}
 	}
 
 	bReinitialized = false;
 
-	TRACE_ANIM_NODE_VALUE(Context, TEXT("Name"), CurrentSequence != nullptr ? CurrentSequence->GetFName() : NAME_None);
-	TRACE_ANIM_NODE_VALUE(Context, TEXT("Sequence"), CurrentSequence);
-	TRACE_ANIM_NODE_VALUE(Context, TEXT("InputTime"), CurrentExplicitTime);
+	TRACE_ANIM_NODE_VALUE(Context, TEXT("Name"), Sequence != nullptr ? Sequence->GetFName() : NAME_None);
+	TRACE_ANIM_NODE_VALUE(Context, TEXT("Sequence"), Sequence);
+	TRACE_ANIM_NODE_VALUE(Context, TEXT("InputTime"), ExplicitTime);
 	TRACE_ANIM_NODE_VALUE(Context, TEXT("Time"), InternalTimeAccumulator);
 }
 
-void FAnimNode_SequenceEvaluatorBase::Evaluate_AnyThread(FPoseContext& Output)
+void FAnimNode_SequenceEvaluator::Evaluate_AnyThread(FPoseContext& Output)
 {
 	DECLARE_SCOPE_HIERARCHICAL_COUNTER_ANIMNODE(Evaluate_AnyThread)
 	check(Output.AnimInstanceProxy != nullptr);
-	UAnimSequenceBase* CurrentSequence = GetSequence();
-	if ((CurrentSequence != nullptr) && (Output.AnimInstanceProxy->IsSkeletonCompatible(CurrentSequence->GetSkeleton())))
+	if ((Sequence != nullptr) && Output.AnimInstanceProxy->IsSkeletonCompatible(Sequence->GetSkeleton()))
 	{
 		FAnimationPoseData AnimationPoseData(Output);
-		CurrentSequence->GetAnimationPose(AnimationPoseData, FAnimExtractContext(InternalTimeAccumulator, Output.AnimInstanceProxy->ShouldExtractRootMotion()));
+		Sequence->GetAnimationPose(AnimationPoseData, FAnimExtractContext(InternalTimeAccumulator, Output.AnimInstanceProxy->ShouldExtractRootMotion()));
 	}
 	else
 	{
@@ -111,49 +108,19 @@ void FAnimNode_SequenceEvaluatorBase::Evaluate_AnyThread(FPoseContext& Output)
 	}
 }
 
-void FAnimNode_SequenceEvaluatorBase::GatherDebugData(FNodeDebugData& DebugData)
+void FAnimNode_SequenceEvaluator::OverrideAsset(UAnimationAsset* NewAsset)
+{
+	if(UAnimSequenceBase* NewSequence = Cast<UAnimSequenceBase>(NewAsset))
+	{
+		Sequence = NewSequence;
+	}
+}
+
+void FAnimNode_SequenceEvaluator::GatherDebugData(FNodeDebugData& DebugData)
 {
 	DECLARE_SCOPE_HIERARCHICAL_COUNTER_ANIMNODE(GatherDebugData)
 	FString DebugLine = DebugData.GetNodeName(this);
 	
-	DebugLine += FString::Printf(TEXT("('%s' InputTime: %.3f, Time: %.3f)"), *GetNameSafe(GetSequence()), GetExplicitTime(), InternalTimeAccumulator);
+	DebugLine += FString::Printf(TEXT("('%s' InputTime: %.3f, Time: %.3f)"), *GetNameSafe(Sequence), ExplicitTime, InternalTimeAccumulator);
 	DebugData.AddDebugItem(DebugLine, true);
 }
-
-void FAnimNode_SequenceEvaluator::SetSequence(UAnimSequenceBase* InSequence)
-{
-#if WITH_EDITORONLY_DATA
-	TObjectPtr<UAnimSequenceBase>& SequenceToSet = GET_MUTABLE_ANIM_NODE_DATA(TObjectPtr<UAnimSequenceBase>, Sequence);
-	SequenceToSet = InSequence;
-#endif
-}
-
-UAnimSequenceBase* FAnimNode_SequenceEvaluator::GetSequence() const
-{
-	return GET_ANIM_NODE_DATA(TObjectPtr<UAnimSequenceBase>, Sequence);
-}
-
-float FAnimNode_SequenceEvaluator::GetExplicitTime() const
-{
-	return GET_ANIM_NODE_DATA(float, ExplicitTime);
-}
-
-bool FAnimNode_SequenceEvaluator::GetShouldLoop() const
-{
-	return GET_ANIM_NODE_DATA(bool, bShouldLoop);
-}
-
-bool FAnimNode_SequenceEvaluator::GetTeleportToExplicitTime() const
-{
-	return GET_ANIM_NODE_DATA(bool, bTeleportToExplicitTime);
-}
-
-TEnumAsByte<ESequenceEvalReinit::Type> FAnimNode_SequenceEvaluator::GetReinitializationBehavior() const
-{
-	return GET_ANIM_NODE_DATA(TEnumAsByte<ESequenceEvalReinit::Type>, ReinitializationBehavior);
-}
-
-float FAnimNode_SequenceEvaluator::GetStartPosition() const
-{
-	return GET_ANIM_NODE_DATA(float, StartPosition);
-}
diff --git a/Engine/Source/Runtime/AnimGraphRuntime/Private/AnimSequencerInstanceProxy.cpp b/Engine/Source/Runtime/AnimGraphRuntime/Private/AnimSequencerInstanceProxy.cpp
--- a/Engine/Source/Runtime/AnimGraphRuntime/Private/AnimSequencerInstanceProxy.cpp
+++ b/Engine/Source/Runtime/AnimGraphRuntime/Private/AnimSequencerInstanceProxy.cpp
@@ -132,7 +132,8 @@ void FAnimSequencerInstanceProxy::InitAnimTrack(UAnimSequenceBase* InAnimSequenc
 			SequencerToPlayerMap.Add(SequenceId, NewPlayerState);
 
 			// link player to blendnode, this will let you trigger notifies and so on
-			NewPlayerState->PlayerNode.SetTeleportToExplicitTime(false);
+			NewPlayerState->PlayerNode.bTeleportToExplicitTime = false;
+			NewPlayerState->PlayerNode.bShouldLoop = true;
 			BlendNode.Poses[PoseIndex].SetLinkNode(&NewPlayerState->PlayerNode);
 
 			// set player state
@@ -140,8 +141,8 @@ void FAnimSequencerInstanceProxy::InitAnimTrack(UAnimSequenceBase* InAnimSequenc
 		}
 
 		// now set animation data to player
-		PlayerState->PlayerNode.SetSequence(InAnimSequence);
-		PlayerState->PlayerNode.SetExplicitTime(0.f);
+		PlayerState->PlayerNode.Sequence = InAnimSequence;
+		PlayerState->PlayerNode.ExplicitTime = 0.f;
 
 		// initialize player
 		PlayerState->PlayerNode.Initialize_AnyThread(FAnimationInitializeContext(this));
@@ -193,17 +194,17 @@ void FAnimSequencerInstanceProxy::UpdateAnimTrack(UAnimSequenceBase* InAnimSeque
 
 	FSequencerPlayerAnimSequence* PlayerState = FindPlayer<FSequencerPlayerAnimSequence>(SequenceId);
 
-	PlayerState->PlayerNode.SetExplicitTime(InToPosition);
+	PlayerState->PlayerNode.ExplicitTime = InToPosition;
 	if (InFromPosition.IsSet())
 	{
 		// Set the internal time accumulator at the "from" time so that the player node will correctly evaluate the
 		// desired "from/to" range. We also disable the reinitialization code so it doesn't mess up that time we
 		// just set.
 		PlayerState->PlayerNode.SetExplicitPreviousTime(InFromPosition.GetValue());
-		PlayerState->PlayerNode.SetReinitializationBehavior(ESequenceEvalReinit::NoReset);
+		PlayerState->PlayerNode.ReinitializationBehavior = ESequenceEvalReinit::NoReset;
 	}
 	// if no fire notifies, we can teleport to explicit time
-	PlayerState->PlayerNode.SetTeleportToExplicitTime(!bFireNotifies);
+	PlayerState->PlayerNode.bTeleportToExplicitTime = !bFireNotifies;
 	// if moving to 0.f, we mark this to teleport. Otherwise, do not use explicit time
 	FAnimNode_MultiWayBlend& BlendNode = (PlayerState->bAdditive) ? AdditiveBlendNode : FullBodyBlendNode;
 	BlendNode.DesiredAlphas[PlayerState->PoseIndex] = Weight;
@@ -223,9 +224,9 @@ void FAnimSequencerInstanceProxy::EnsureAnimTrack(UAnimSequenceBase* InAnimSeque
 	{
 		InitAnimTrack(InAnimSequence, SequenceId);
 	}
-	else if (PlayerState->PlayerNode.GetSequence() != InAnimSequence)
+	else if (PlayerState->PlayerNode.Sequence != InAnimSequence)
 	{
-		PlayerState->PlayerNode.SetSequence(InAnimSequence);
+		PlayerState->PlayerNode.OverrideAsset(InAnimSequence);
 	}
 }
 
diff --git a/Engine/Source/Runtime/AnimGraphRuntime/Public/AnimNodes/AnimNode_AimOffsetLookAt.h b/Engine/Source/Runtime/AnimGraphRuntime/Public/AnimNodes/AnimNode_AimOffsetLookAt.h
--- a/Engine/Source/Runtime/AnimGraphRuntime/Public/AnimNodes/AnimNode_AimOffsetLookAt.h
+++ b/Engine/Source/Runtime/AnimGraphRuntime/Public/AnimNodes/AnimNode_AimOffsetLookAt.h
@@ -55,9 +55,6 @@ struct ANIMGRAPHRUNTIME_API FAnimNode_AimOffsetLookAt : public FAnimNode_BlendSp
 	UPROPERTY(EditAnywhere, BlueprintReadWrite, Category = LookAt, meta = (PinShownByDefault))
 	float Alpha;
 
-	/** Cached calculated blend input */
-	FVector CurrentBlendInput;
-
 	/** Cached reference to the source socket's bone */
 	FBoneReference SocketBoneReference;
 
@@ -81,8 +78,5 @@ public:
 	virtual int32 GetLODThreshold() const override { return LODThreshold; }
 	// End of FAnimNode_Base interface
 
-	// FAnimNode_BlendSpacePlayer interface
-	virtual FVector GetPosition() const override;
-
 	void UpdateFromLookAtTarget(FPoseContext& LocalPoseContext);
 };
diff --git a/Engine/Source/Runtime/AnimGraphRuntime/Public/AnimNodes/AnimNode_BlendListBase.h b/Engine/Source/Runtime/AnimGraphRuntime/Public/AnimNodes/AnimNode_BlendListBase.h
--- a/Engine/Source/Runtime/AnimGraphRuntime/Public/AnimNodes/AnimNode_BlendListBase.h
+++ b/Engine/Source/Runtime/AnimGraphRuntime/Public/AnimNodes/AnimNode_BlendListBase.h
@@ -23,36 +23,34 @@ enum class EBlendListTransitionType : uint8
 USTRUCT(BlueprintInternalUseOnly)
 struct ANIMGRAPHRUNTIME_API FAnimNode_BlendListBase : public FAnimNode_Base
 {
-	GENERATED_BODY()
-
-protected:	
-	UPROPERTY(EditAnywhere, EditFixedSize, Category=Links)
+	GENERATED_USTRUCT_BODY()
+public:
+	UPROPERTY(EditAnywhere, EditFixedSize, BlueprintReadWrite, Category=Links)
 	TArray<FPoseLink> BlendPose;
 
-private:
-#if WITH_EDITORONLY_DATA
-	UPROPERTY(EditAnywhere, EditFixedSize, Category=Config, meta=(PinShownByDefault, FoldProperty))
+	UPROPERTY(EditAnywhere, EditFixedSize, BlueprintReadWrite, Category=Config, meta=(PinShownByDefault))
 	TArray<float> BlendTime;
 
-	UPROPERTY(EditAnywhere, Category=Config, meta=(FoldProperty))
-	EBlendListTransitionType TransitionType = EBlendListTransitionType::StandardBlend;
+	UPROPERTY(EditAnywhere, BlueprintReadWrite, Category=Config)
+	EBlendListTransitionType TransitionType;
 
-	UPROPERTY(EditAnywhere, Category=BlendType, meta=(FoldProperty))
-	EAlphaBlendOption BlendType = EAlphaBlendOption::Linear;
+	UPROPERTY(EditAnywhere, Category=BlendType)
+	EAlphaBlendOption BlendType;
 	
 protected:
 	/** This reinitializes child pose when re-activated. For example, when active child changes */
-	UPROPERTY(EditAnywhere, Category = Option, meta=(FoldProperty))
-	bool bResetChildOnActivation = false;
+	UPROPERTY(EditAnywhere, Category = Option)
+	bool bResetChildOnActivation;
 
-private:
-	UPROPERTY(EditAnywhere, Category=BlendType, meta=(FoldProperty))
-	TObjectPtr<UCurveFloat> CustomBlendCurve = nullptr;
+	int32 LastActiveChildIndex;
+
+public:
+	UPROPERTY(EditAnywhere, Category=BlendType)
+	TObjectPtr<UCurveFloat> CustomBlendCurve;
+
+	UPROPERTY(EditAnywhere, Category=BlendType, meta=(UseAsBlendProfile=true))
+	TObjectPtr<UBlendProfile> BlendProfile;
 
-	UPROPERTY(EditAnywhere, Category=BlendType, meta=(UseAsBlendProfile=true, FoldProperty))
-	TObjectPtr<UBlendProfile> BlendProfile = nullptr;
-#endif // #if WITH_EDITORONLY_DATA
-	
 	TArray<struct FAlphaBlend> Blends;
 
 protected:
@@ -67,10 +65,16 @@ protected:
 	//Store which poses we need to evaluate
 	TArray<int32> PosesToEvaluate;
 
-	int32 LastActiveChildIndex = 0;
-	
 public:	
-	FAnimNode_BlendListBase() = default;
+	FAnimNode_BlendListBase()
+		: TransitionType(EBlendListTransitionType::StandardBlend)
+		, BlendType(EAlphaBlendOption::Linear)
+		, bResetChildOnActivation(false)
+		, LastActiveChildIndex(0)
+		, CustomBlendCurve(nullptr)
+		, BlendProfile(nullptr)
+	{
+	}
 
 	// FAnimNode_Base interface
 	virtual void Initialize_AnyThread(const FAnimationInitializeContext& Context) override;
@@ -94,25 +98,6 @@ public:
 	}
 #endif
 
-public:
-	// Get the array of blend times to apply to our input poses
-	const TArray<float>& GetBlendTimes() const;
-
-	// Get the type of transition that this blend list will make
-	EBlendListTransitionType GetTransitionType() const;
-
-	// Get the blend type we will use when blending
-	EAlphaBlendOption GetBlendType() const;
-	
-	/** Get whether to reinitialize the child pose when re-activated. For example, when active child changes */
-	bool GetResetChildOnActivation() const;
-
-	// Get the custom blend curve to apply when blending, if any
-	UCurveFloat* GetCustomBlendCurve() const;
-
-	// Get the blend profile to use when blending, if any
-	UBlendProfile* GetBlendProfile() const;
-	
 protected:
 	virtual int32 GetActiveChildIndex() { return 0; }
 	virtual FString GetNodeName(FNodeDebugData& DebugData) { return DebugData.GetNodeName(this); }
diff --git a/Engine/Source/Runtime/AnimGraphRuntime/Public/AnimNodes/AnimNode_BlendListByBool.h b/Engine/Source/Runtime/AnimGraphRuntime/Public/AnimNodes/AnimNode_BlendListByBool.h
--- a/Engine/Source/Runtime/AnimGraphRuntime/Public/AnimNodes/AnimNode_BlendListByBool.h
+++ b/Engine/Source/Runtime/AnimGraphRuntime/Public/AnimNodes/AnimNode_BlendListByBool.h
@@ -11,20 +11,19 @@
 USTRUCT(BlueprintInternalUseOnly)
 struct ANIMGRAPHRUNTIME_API FAnimNode_BlendListByBool : public FAnimNode_BlendListBase
 {
-	GENERATED_BODY()
-private:
-#if WITH_EDITORONLY_DATA
+	GENERATED_USTRUCT_BODY()
+public:
 	// Which input should be connected to the output?
-	UPROPERTY(EditAnywhere, Category=Runtime, meta=(PinShownByDefault, FoldProperty))
-	bool bActiveValue = false;
-#endif
-	
+	UPROPERTY(EditAnywhere, BlueprintReadWrite, Category=Runtime, meta=(PinShownByDefault))
+	bool bActiveValue;
+
 public:	
-	FAnimNode_BlendListByBool() = default;
+	FAnimNode_BlendListByBool()
+		: FAnimNode_BlendListBase()
+		, bActiveValue(false)
+	{
+	}
 
-	// Get which input should be connected to the output
-	bool GetActiveValue() const;
-	
 protected:
 	virtual int32 GetActiveChildIndex() override;
 	virtual FString GetNodeName(FNodeDebugData& DebugData) override { return DebugData.GetNodeName(this); }
diff --git a/Engine/Source/Runtime/AnimGraphRuntime/Public/AnimNodes/AnimNode_BlendListByEnum.h b/Engine/Source/Runtime/AnimGraphRuntime/Public/AnimNodes/AnimNode_BlendListByEnum.h
--- a/Engine/Source/Runtime/AnimGraphRuntime/Public/AnimNodes/AnimNode_BlendListByEnum.h
+++ b/Engine/Source/Runtime/AnimGraphRuntime/Public/AnimNodes/AnimNode_BlendListByEnum.h
@@ -11,33 +11,23 @@
 USTRUCT(BlueprintInternalUseOnly)
 struct ANIMGRAPHRUNTIME_API FAnimNode_BlendListByEnum : public FAnimNode_BlendListBase
 {
-	GENERATED_BODY()
-	
-private:
-#if WITH_EDITORONLY_DATA
+	GENERATED_USTRUCT_BODY()
+public:
 	// Mapping from enum value to BlendPose index; there will be one entry per entry in the enum; entries out of range always map to pose index 0
-	UPROPERTY(meta=(FoldProperty))
+	UPROPERTY()
 	TArray<int32> EnumToPoseIndex;
-	
+
 	// The currently selected pose (as an enum value)
-	UPROPERTY(EditAnywhere, Category=Runtime, meta=(PinShownByDefault, FoldProperty))
-	mutable uint8 ActiveEnumValue = 0;
-#endif
-	
+	UPROPERTY(EditAnywhere, BlueprintReadWrite, Category=Runtime, meta=(PinShownByDefault))
+	mutable uint8 ActiveEnumValue;
+
 public:	
-	FAnimNode_BlendListByEnum() = default;
+	FAnimNode_BlendListByEnum()
+		: FAnimNode_BlendListBase()
+		, ActiveEnumValue(0)
+	{
+	}
 
-#if WITH_EDITORONLY_DATA
-	// Set the mapping from enum value to BlendPose index. Called during compilation.
-	void SetEnumToPoseIndex(const TArray<int32>& InEnumToPoseIndex);
-#endif
-	
-	// Get the mapping from enum value to BlendPose index; there will be one entry per entry in the enum; entries out of range always map to pose index 0
-	const TArray<int32>& GetEnumToPoseIndex() const;
-	
-	// Get the currently selected pose (as an enum value)
-	uint8 GetActiveEnumValue() const;
-	
 protected:
 	virtual int32 GetActiveChildIndex() override;
 	virtual FString GetNodeName(FNodeDebugData& DebugData) override { return DebugData.GetNodeName(this); }
diff --git a/Engine/Source/Runtime/AnimGraphRuntime/Public/AnimNodes/AnimNode_BlendListByInt.h b/Engine/Source/Runtime/AnimGraphRuntime/Public/AnimNodes/AnimNode_BlendListByInt.h
--- a/Engine/Source/Runtime/AnimGraphRuntime/Public/AnimNodes/AnimNode_BlendListByInt.h
+++ b/Engine/Source/Runtime/AnimGraphRuntime/Public/AnimNodes/AnimNode_BlendListByInt.h
@@ -11,19 +11,19 @@
 USTRUCT(BlueprintInternalUseOnly)
 struct ANIMGRAPHRUNTIME_API FAnimNode_BlendListByInt : public FAnimNode_BlendListBase
 {
-	GENERATED_BODY()
-private:
-#if WITH_EDITORONLY_DATA
-	UPROPERTY(EditAnywhere,  Category=Runtime, meta=(PinShownByDefault, FoldProperty))
-	int32 ActiveChildIndex = 0;
-#endif
-	
+	GENERATED_USTRUCT_BODY()
 public:
-	FAnimNode_BlendListByInt() = default;
+	UPROPERTY(EditAnywhere, BlueprintReadWrite, Category=Runtime, meta=(PinShownByDefault))
+	int32 ActiveChildIndex;
+
+public:	
+	FAnimNode_BlendListByInt()
+		:	FAnimNode_BlendListBase()
+		,	ActiveChildIndex(0)
+	{
+	}
 
-	// Get the currently active child index
-	virtual int32 GetActiveChildIndex() override;
-	
 protected:
-	virtual FString GetNodeName(FNodeDebugData& DebugData) override { return DebugData.GetNodeName(this); }
+	virtual int32 GetActiveChildIndex();
+	virtual FString GetNodeName(FNodeDebugData& DebugData) { return DebugData.GetNodeName(this); }
 };
diff --git a/Engine/Source/Runtime/AnimGraphRuntime/Public/AnimNodes/AnimNode_BlendSpacePlayer.h b/Engine/Source/Runtime/AnimGraphRuntime/Public/AnimNodes/AnimNode_BlendSpacePlayer.h
--- a/Engine/Source/Runtime/AnimGraphRuntime/Public/AnimNodes/AnimNode_BlendSpacePlayer.h
+++ b/Engine/Source/Runtime/AnimGraphRuntime/Public/AnimNodes/AnimNode_BlendSpacePlayer.h
@@ -14,73 +14,67 @@ class UBlendSpace;
 USTRUCT(BlueprintInternalUseOnly)
 struct ANIMGRAPHRUNTIME_API FAnimNode_BlendSpacePlayer : public FAnimNode_AssetPlayerBase
 {
-	GENERATED_BODY()
+	GENERATED_USTRUCT_BODY()
 
- 	friend class UAnimGraphNode_BlendSpacePlayer;
- 	friend class UAnimGraphNode_BlendSpaceEvaluator;
-	friend class UAnimGraphNode_RotationOffsetBlendSpace;
-	friend class UAnimGraphNode_AimOffsetLookAt;
+	// @return the current sample coordinates that this node is using to sample the blendspace
+	FVector GetPosition() const { return FVector(X, Y, Z); }
 
 	// @return the current sample coordinates after going through the filtering
 	FVector GetFilteredPosition() const { return BlendFilter.GetFilterLastOutput(); }
-private:
 
-#if WITH_EDITORONLY_DATA
+public:
 	// The X coordinate to sample in the blendspace
-	UPROPERTY(EditAnywhere, Category=Coordinates, meta=(PinShownByDefault, FoldProperty))
-	float X = 0.0f;
+	UPROPERTY(EditAnywhere, BlueprintReadWrite, Category=Coordinates, meta=(PinShownByDefault))
+	float X;
 
 	// The Y coordinate to sample in the blendspace
-	UPROPERTY(EditAnywhere, Category = Coordinates, meta = (PinShownByDefault, FoldProperty))
-	float Y = 0.0f;
+	UPROPERTY(EditAnywhere, BlueprintReadWrite, Category = Coordinates, meta = (PinShownByDefault))
+	float Y;
 
 	// The Z coordinate to sample in the blendspace
-	UPROPERTY(EditAnywhere, Category = Coordinates, meta = (PinHiddenByDefault, FoldProperty))
-	float Z = 0.0f;
+	UPROPERTY(EditAnywhere, BlueprintReadWrite, Category = Coordinates, meta = (PinHiddenByDefault))
+	float Z;
 
 	// The play rate multiplier. Can be negative, which will cause the animation to play in reverse.
-	UPROPERTY(EditAnywhere, Category = Settings, meta = (DefaultValue = "1.0", PinHiddenByDefault, FoldProperty))
-	float PlayRate = 1.0f;
+	UPROPERTY(EditAnywhere, BlueprintReadWrite, Category = Settings, meta = (DefaultValue = "1.0", PinHiddenByDefault))
+	float PlayRate;
 
 	// Should the animation continue looping when it reaches the end?
-	UPROPERTY(EditAnywhere, Category = Settings, meta = (DefaultValue = "true", PinHiddenByDefault, FoldProperty))
-	bool bLoop = true;
+	UPROPERTY(EditAnywhere, BlueprintReadWrite, Category = Settings, meta = (DefaultValue = "true", PinHiddenByDefault))
+	bool bLoop;
 
 	// Whether we should reset the current play time when the blend space changes
-	UPROPERTY(EditAnywhere, Category = Settings, meta = (PinHiddenByDefault, FoldProperty))
-	bool bResetPlayTimeWhenBlendSpaceChanges = true;
+	UPROPERTY(EditAnywhere, BlueprintReadWrite, Category = Settings, meta = (PinHiddenByDefault))
+	bool bResetPlayTimeWhenBlendSpaceChanges;
 
 	// The start up position in [0, 1], it only applies when reinitialized
 	// if you loop, it will still start from 0.f after finishing the round
-	UPROPERTY(EditAnywhere, Category=Settings, meta = (DefaultValue = "0.f", PinHiddenByDefault, FoldProperty))
-	float StartPosition = 0.0f;
+	UPROPERTY(EditAnywhere, BlueprintReadWrite, Category=Settings, meta = (DefaultValue = "0.f", PinHiddenByDefault))
+	float StartPosition;
 
 	// The blendspace asset to play
-	UPROPERTY(EditAnywhere, Category = Settings, meta = (PinHiddenByDefault, FoldProperty))
-	TObjectPtr<UBlendSpace> BlendSpace = nullptr;
-#endif
+	UPROPERTY(EditAnywhere, BlueprintReadWrite, Category = Settings, meta = (PinHiddenByDefault))
+	TObjectPtr<UBlendSpace> BlendSpace;
 
 protected:
-	// Filter used to dampen coordinate changes
+
 	FBlendFilter BlendFilter;
 
-	// Cache of samples used to determine blend weights
 	TArray<FBlendSampleData> BlendSampleDataCache;
 
 	/** Previous position in the triangulation/segmentation */
 	int32 CachedTriangulationIndex = -1;
 
 	UPROPERTY(Transient)
-	TObjectPtr<UBlendSpace> PreviousBlendSpace = nullptr;
+	TObjectPtr<UBlendSpace> PreviousBlendSpace;
 
 public:	
-	FAnimNode_BlendSpacePlayer() = default;
+	FAnimNode_BlendSpacePlayer();
 
 	// FAnimNode_AssetPlayerBase interface
-	virtual float GetCurrentAssetTime() const override;
-	virtual float GetCurrentAssetTimePlayRateAdjusted() const override;
-	virtual float GetCurrentAssetLength() const override;
-	virtual UAnimationAsset* GetAnimAsset() const override;
+	virtual float GetCurrentAssetTime();
+	virtual float GetCurrentAssetTimePlayRateAdjusted();
+	virtual float GetCurrentAssetLength();
 	// End of FAnimNode_AssetPlayerBase interface
 
 	// FAnimNode_Base interface
@@ -88,34 +82,13 @@ public:
 	virtual void CacheBones_AnyThread(const FAnimationCacheBonesContext& Context) override;
 	virtual void UpdateAssetPlayer(const FAnimationUpdateContext& Context) override;
 	virtual void Evaluate_AnyThread(FPoseContext& Output) override;
+	virtual void OverrideAsset(UAnimationAsset* NewAsset) override;
 	virtual void GatherDebugData(FNodeDebugData& DebugData) override;
 	// End of FAnimNode_Base interface
 
-	float GetTimeFromEnd(float CurrentTime) const;
-
-#if WITH_EDITORONLY_DATA
-	// Set the blendspace asset to play
-	void SetBlendSpace(UBlendSpace* InBlendSpace);
-#endif
-
-	// Get the coordinates that are currently being sampled by the blendspace
-	virtual FVector GetPosition() const;
-
-	// Get the play rate multiplier. Can be negative, which will cause the animation to play in reverse.
-	float GetPlayRate() const;
-
-	// Should the animation continue looping when it reaches the end?
-	bool GetLoop() const;
-
-	// Get whether we should reset the current play time when the blend space changes
-	bool ShouldResetPlayTimeWhenBlendSpaceChanges() const;
-
-	// Get the start up position in [0, 1], it only applies when reinitialized
-	// if you loop, it will still start from 0.f after finishing the round
-	float GetStartPosition() const;
+	float GetTimeFromEnd(float CurrentTime);
 
-	// Get the blendspace asset to play
-	UBlendSpace* GetBlendSpace() const;
+	UAnimationAsset* GetAnimAsset();
 
 protected:
 	void UpdateInternal(const FAnimationUpdateContext& Context);
diff --git a/Engine/Source/Runtime/AnimGraphRuntime/Public/AnimNodes/AnimNode_PoseHandler.h b/Engine/Source/Runtime/AnimGraphRuntime/Public/AnimNodes/AnimNode_PoseHandler.h
--- a/Engine/Source/Runtime/AnimGraphRuntime/Public/AnimNodes/AnimNode_PoseHandler.h
+++ b/Engine/Source/Runtime/AnimGraphRuntime/Public/AnimNodes/AnimNode_PoseHandler.h
@@ -28,28 +28,24 @@ public:
 	}
 
 	// FAnimNode_AssetPlayerBase interface
-	virtual float GetCurrentAssetTime() const { return 0.f; }
-	virtual float GetCurrentAssetLength() const { return 0.f; }
+	virtual float GetCurrentAssetTime() { return 0.f; }
+	virtual float GetCurrentAssetLength() { return 0.f; }
 	// End of FAnimNode_AssetPlayerBase interface
 
 	// FAnimNode_Base interface
 	virtual void Initialize_AnyThread(const FAnimationInitializeContext& Context) override;
 	virtual void CacheBones_AnyThread(const FAnimationCacheBonesContext& Context) override;
 	virtual void UpdateAssetPlayer(const FAnimationUpdateContext& Context) override;
+	virtual void OverrideAsset(UAnimationAsset* NewAsset) override;
 	virtual void GatherDebugData(FNodeDebugData& DebugData) override;
 	// End of FAnimNode_Base interface
 
 	// FAnimNode_AssetPlayerBase Interface
 	virtual float GetAccumulatedTime() const {return 0.f;}
 	virtual void SetAccumulatedTime(const float& NewTime) {}
-	virtual UAnimationAsset* GetAnimAsset() const {return PoseAsset;}
+	virtual UAnimationAsset* GetAnimAsset() {return PoseAsset;}
 	// End of FAnimNode_AssetPlayerBase Interface
 
-#if WITH_EDITORONLY_DATA
-	// Set the pose asset to use for this node 
-	void SetPoseAsset(UPoseAsset* InPoseAsset);
-#endif
-	
 protected:
 	/** Called after CurrentPoseAsset is changed.  */
 	virtual void OnPoseAssetChange() {}
diff --git a/Engine/Source/Runtime/AnimGraphRuntime/Public/AnimNodes/AnimNode_RefPose.h b/Engine/Source/Runtime/AnimGraphRuntime/Public/AnimNodes/AnimNode_RefPose.h
--- a/Engine/Source/Runtime/AnimGraphRuntime/Public/AnimNodes/AnimNode_RefPose.h
+++ b/Engine/Source/Runtime/AnimGraphRuntime/Public/AnimNodes/AnimNode_RefPose.h
@@ -18,35 +18,32 @@ enum ERefPoseType
 USTRUCT()
 struct ANIMGRAPHRUNTIME_API FAnimNode_RefPose : public FAnimNode_Base
 {
-	GENERATED_BODY()
+	GENERATED_USTRUCT_BODY()
 
-private:
-#if WITH_EDITORONLY_DATA
-	UPROPERTY(meta=(FoldProperty))
-	TEnumAsByte<ERefPoseType> RefPoseType = EIT_LocalSpace;
-#endif	// #if WITH_EDITORONLY_DATA
+	UPROPERTY()
+	TEnumAsByte<ERefPoseType> RefPoseType;
 	
 public:	
-	FAnimNode_RefPose() = default;
+	FAnimNode_RefPose()
+		:	FAnimNode_Base()
+		, 	RefPoseType(EIT_LocalSpace)
+	{
+	}
 
-#if WITH_EDITORONLY_DATA
-	// Set the ref pose type of this node
-	void SetRefPoseType(ERefPoseType InType) { RefPoseType = InType; }
-#endif
-
-	// Get the type of this ref pose
-	ERefPoseType GetRefPoseType() const;
-	
 	virtual void Evaluate_AnyThread(FPoseContext& Output) override;
 	virtual void GatherDebugData(FNodeDebugData& DebugData) override;
+
 };
 
 USTRUCT()
 struct ANIMGRAPHRUNTIME_API FAnimNode_MeshSpaceRefPose : public FAnimNode_Base
 {
-	GENERATED_BODY()
+	GENERATED_USTRUCT_BODY()
 public:	
-	FAnimNode_MeshSpaceRefPose() = default;
+	FAnimNode_MeshSpaceRefPose()
+		:	FAnimNode_Base()
+	{
+	}
 
 	virtual void EvaluateComponentSpace_AnyThread(FComponentSpacePoseContext& Output);
 };
