Revert optimization for layered bone per blend

This can cause crashes in certain graph configurations (e.g. when within a state machine, Initialize can be called without a CacheBones call directly afterwards)

diff --git a/Engine/Source/Runtime/AnimGraphRuntime/Private/AnimNodes/AnimNode_LayeredBoneBlend.cpp b/Engine/Source/Runtime/AnimGraphRuntime/Private/AnimNodes/AnimNode_LayeredBoneBlend.cpp
--- a/Engine/Source/Runtime/AnimGraphRuntime/Private/AnimNodes/AnimNode_LayeredBoneBlend.cpp
+++ b/Engine/Source/Runtime/AnimGraphRuntime/Private/AnimNodes/AnimNode_LayeredBoneBlend.cpp
@@ -26,6 +26,10 @@ void FAnimNode_LayeredBoneBlend::Initialize_AnyThread(const FAnimationInitialize
 		{
 			BlendPoses[ChildIndex].Initialize(Context);
 		}
+		
+		// initialize mask weight now
+    	check (Context.AnimInstanceProxy->GetSkeleton());
+    	ReinitializeBoneBlendWeights(Context.AnimInstanceProxy->GetRequiredBones(), Context.AnimInstanceProxy->GetSkeleton());
 	}
 }
 
