Backing out CL 16112742

diff --git a/Engine/Source/Runtime/AnimGraphRuntime/Private/AnimNodes/AnimNode_LayeredBoneBlend.cpp b/Engine/Source/Runtime/AnimGraphRuntime/Private/AnimNodes/AnimNode_LayeredBoneBlend.cpp
--- a/Engine/Source/Runtime/AnimGraphRuntime/Private/AnimNodes/AnimNode_LayeredBoneBlend.cpp
+++ b/Engine/Source/Runtime/AnimGraphRuntime/Private/AnimNodes/AnimNode_LayeredBoneBlend.cpp
@@ -26,10 +26,6 @@ void FAnimNode_LayeredBoneBlend::Initialize_AnyThread(const FAnimationInitialize
 		{
 			BlendPoses[ChildIndex].Initialize(Context);
 		}
-		
-		// initialize mask weight now
-    	check (Context.AnimInstanceProxy->GetSkeleton());
-    	ReinitializeBoneBlendWeights(Context.AnimInstanceProxy->GetRequiredBones(), Context.AnimInstanceProxy->GetSkeleton());
 	}
 }
 
