Added animation mirroring to sequencer

diff --git a/Engine/Source/Editor/MovieSceneTools/Private/TrackEditors/SkeletalAnimationTrackEditor.cpp b/Engine/Source/Editor/MovieSceneTools/Private/TrackEditors/SkeletalAnimationTrackEditor.cpp
--- a/Engine/Source/Editor/MovieSceneTools/Private/TrackEditors/SkeletalAnimationTrackEditor.cpp
+++ b/Engine/Source/Editor/MovieSceneTools/Private/TrackEditors/SkeletalAnimationTrackEditor.cpp
@@ -29,6 +29,7 @@
 #include "Animation/AnimMontage.h"
 #include "Animation/AnimSequence.h"
 #include "Animation/PoseAsset.h"
+#include "Animation/MirrorDataTable.h"
 #include "EditorStyleSet.h"
 #include "DragAndDrop/AssetDragDropOp.h"
 #include "MovieSceneTimeHelpers.h"
@@ -410,6 +411,7 @@ public:
 	virtual void CustomizeChildren(TSharedRef<IPropertyHandle> PropertyHandle, IDetailChildrenBuilder& ChildBuilder, IPropertyTypeCustomizationUtils& CustomizationUtils) override
 	{
 		const FName AnimationPropertyName = GET_MEMBER_NAME_CHECKED(FMovieSceneSkeletalAnimationParams, Animation);
+		const FName MirrorDataTableName = GET_MEMBER_NAME_CHECKED(FMovieSceneSkeletalAnimationParams, MirrorDataTable);
 
 		uint32 NumChildren;
 		PropertyHandle->GetNumChildren(NumChildren);
@@ -417,10 +419,10 @@ public:
 		{
 			TSharedPtr<IPropertyHandle> ChildPropertyHandle = PropertyHandle->GetChildHandle(i);
 			IDetailPropertyRow& ChildPropertyRow = ChildBuilder.AddProperty(ChildPropertyHandle.ToSharedRef());
-
-			// Let most properties be whatever they want to be... we just want to customize the `Animation` property
+			FName ChildPropertyName = ChildPropertyHandle->GetProperty()->GetFName();
+			// Let most properties be whatever they want to be... we just want to customize the `Animation` and `MirrorDataTable` properties
 			// by making it look like a normal asset reference property, but with some custom filtering.
-			if (ChildPropertyHandle->GetProperty()->GetFName() == AnimationPropertyName)
+			if (ChildPropertyName == AnimationPropertyName || ChildPropertyName == MirrorDataTableName)
 			{
 				FDetailWidgetRow& Row = ChildPropertyRow.CustomWidget();
 
@@ -432,10 +434,11 @@ public:
 					SkeletonName = FAssetData(Skeleton).GetExportTextName();
 
 					TSharedPtr<IPropertyUtilities> PropertyUtilities = CustomizationUtils.GetPropertyUtilities();
+					UClass* AllowedStaticClass = ChildPropertyName == AnimationPropertyName ? UAnimSequenceBase::StaticClass() : UMirrorDataTable::StaticClass(); 
 
 					TSharedRef<SObjectPropertyEntryBox> ContentWidget = SNew(SObjectPropertyEntryBox)
 						.PropertyHandle(ChildPropertyHandle)
-						.AllowedClass(UAnimSequenceBase::StaticClass())
+						.AllowedClass(AllowedStaticClass)
 						.DisplayThumbnail(true)
 						.ThumbnailPool(PropertyUtilities.IsValid() ? PropertyUtilities->GetThumbnailPool() : nullptr)
 						.OnShouldFilterAsset(FOnShouldFilterAsset::CreateRaw(this, &FMovieSceneSkeletalAnimationParamsDetailCustomization::ShouldFilterAsset));
@@ -505,7 +508,14 @@ FText FSkeletalAnimationSection::GetSectionTitle() const
 {
 	if (Section.Params.Animation != nullptr)
 	{
-		return FText::FromString( Section.Params.Animation->GetName() );
+		if (Section.Params.MirrorDataTable.IsNull())
+		{
+			return FText::FromString( Section.Params.Animation->GetName() );
+		}
+		else
+		{
+			return FText::Format(LOCTEXT("SectionTitleContentFormat", "{0} mirrored with {1}"), FText::FromString(Section.Params.Animation->GetName()), FText::FromString(Section.Params.MirrorDataTable->GetName()));
+		}
 	}
 	return LOCTEXT("NoAnimationSection", "No Animation");
 }
diff --git a/Engine/Source/Runtime/AnimGraphRuntime/Private/AnimNodes/AnimNode_Mirror.cpp b/Engine/Source/Runtime/AnimGraphRuntime/Private/AnimNodes/AnimNode_Mirror.cpp
--- a/Engine/Source/Runtime/AnimGraphRuntime/Private/AnimNodes/AnimNode_Mirror.cpp
+++ b/Engine/Source/Runtime/AnimGraphRuntime/Private/AnimNodes/AnimNode_Mirror.cpp
@@ -10,9 +10,9 @@
 #define LOCTEXT_NAMESPACE "AnimNode_Mirror"
 
 FAnimNode_Mirror::FAnimNode_Mirror()
-	: MirrorDataTable(nullptr)
+	: bMirror(true)
+	, MirrorDataTable(nullptr)
 	, BlendTimeOnMirrorStateChange(0.0f)
-	, bMirror(true)
 	, bBoneMirroring(true)
 	, bCurveMirroring(true)
 	, bAttributeMirroring(true)
@@ -22,7 +22,6 @@ FAnimNode_Mirror::FAnimNode_Mirror()
 {
 }
 
-#if WITH_EDITOR
 UMirrorDataTable* FAnimNode_Mirror::GetMirrorDataTable() const 
 { 
 	return MirrorDataTable.Get();
@@ -33,8 +32,6 @@ void FAnimNode_Mirror::SetMirrorDataTable(UMirrorDataTable* MirrorTable)
 	MirrorDataTable = MirrorTable; 
 }
 
-#endif 
-
 void FAnimNode_Mirror::Initialize_AnyThread(const FAnimationInitializeContext& Context)
 {
 	FAnimNode_Base::Initialize_AnyThread(Context);
@@ -46,12 +43,16 @@ void FAnimNode_Mirror::CacheBones_AnyThread(const FAnimationCacheBonesContext& C
 	Source.CacheBones(Context);
 	DECLARE_SCOPE_HIERARCHICAL_COUNTER_ANIMNODE(CacheBones_AnyThread)
 	Super::CacheBones_AnyThread(Context);
+	const FBoneContainer& BoneContainer = Context.AnimInstanceProxy->GetRequiredBones();
+	FillCompactPoseAndComponentRefRotations(BoneContainer);
+}
 
-	if (MirrorDataTable && Context.AnimInstanceProxy)
+void FAnimNode_Mirror::FillCompactPoseAndComponentRefRotations(const FBoneContainer& BoneContainer)
+{
+	if (MirrorDataTable)
 	{
-		const FBoneContainer& BoneContainer = Context.AnimInstanceProxy->GetRequiredBones();
 		bool bSharesSkeleton = BoneContainer.GetSkeletonAsset() == MirrorDataTable->Skeleton;
-		bool bIdenticalJointCount = BoneContainer.GetReferenceSkeleton().GetNum() == MirrorDataTable->BoneToMirrorBoneIndex.Num(); 
+		bool bIdenticalJointCount = BoneContainer.GetReferenceSkeleton().GetNum() == MirrorDataTable->BoneToMirrorBoneIndex.Num();
 
 		if (bSharesSkeleton && bIdenticalJointCount)
 		{
@@ -76,7 +77,7 @@ void FAnimNode_Mirror::CacheBones_AnyThread(const FAnimationCacheBonesContext& C
 	else
 	{
 		CompactPoseMirrorBones.Reset();
-		ComponentSpaceRefRotations.Reset(); 
+		ComponentSpaceRefRotations.Reset();
 	}
 }
 
@@ -125,6 +126,13 @@ void FAnimNode_Mirror::Evaluate_AnyThread(FPoseContext& Output)
 	{
 		if (bBoneMirroring)
 		{
+			const FBoneContainer& BoneContainer = Output.Pose.GetBoneContainer();
+			const TArray<FBoneIndexType>& RequiredBoneIndices = BoneContainer.GetBoneIndicesArray();
+			int32 NumReqBones = RequiredBoneIndices.Num();
+			if (CompactPoseMirrorBones.Num() != NumReqBones)
+			{
+				FillCompactPoseAndComponentRefRotations(BoneContainer);
+			}
 			FAnimationRuntime::MirrorPose(Output.Pose, MirrorDataTable->MirrorAxis, CompactPoseMirrorBones, ComponentSpaceRefRotations);
 		}
 
diff --git a/Engine/Source/Runtime/AnimGraphRuntime/Private/AnimSequencerInstance.cpp b/Engine/Source/Runtime/AnimGraphRuntime/Private/AnimSequencerInstance.cpp
--- a/Engine/Source/Runtime/AnimGraphRuntime/Private/AnimSequencerInstance.cpp
+++ b/Engine/Source/Runtime/AnimGraphRuntime/Private/AnimSequencerInstance.cpp
@@ -37,7 +37,12 @@ void UAnimSequencerInstance::UpdateAnimTrack(UAnimSequenceBase* InAnimSequence,
 
 void UAnimSequencerInstance::UpdateAnimTrackWithRootMotion(UAnimSequenceBase* InAnimSequence, int32 SequenceId, const TOptional<FRootMotionOverride>& RootMotion, float InFromPosition, float InToPosition, float Weight, bool bFireNotifies)
 {
-	GetProxyOnGameThread<FAnimSequencerInstanceProxy>().UpdateAnimTrackWithRootMotion(InAnimSequence, SequenceId, RootMotion, InFromPosition, InToPosition, Weight, bFireNotifies);
+	UpdateAnimTrackWithRootMotion(InAnimSequence, SequenceId, RootMotion, InFromPosition, InToPosition, Weight, bFireNotifies, nullptr);
+}
+
+void UAnimSequencerInstance::UpdateAnimTrackWithRootMotion(UAnimSequenceBase* InAnimSequence, int32 SequenceId, const TOptional<FRootMotionOverride>& RootMotion, float InFromPosition, float InToPosition, float Weight, bool bFireNotifies, UMirrorDataTable* InMirrorDataTable)
+{
+	GetProxyOnGameThread<FAnimSequencerInstanceProxy>().UpdateAnimTrackWithRootMotion(InAnimSequence, SequenceId, RootMotion, InFromPosition, InToPosition, Weight, bFireNotifies, InMirrorDataTable);
 }
 
 void UAnimSequencerInstance::ConstructNodes()
diff --git a/Engine/Source/Runtime/AnimGraphRuntime/Private/AnimSequencerInstanceProxy.cpp b/Engine/Source/Runtime/AnimGraphRuntime/Private/AnimSequencerInstanceProxy.cpp
--- a/Engine/Source/Runtime/AnimGraphRuntime/Private/AnimSequencerInstanceProxy.cpp
+++ b/Engine/Source/Runtime/AnimGraphRuntime/Private/AnimSequencerInstanceProxy.cpp
@@ -18,7 +18,7 @@ void FAnimSequencerInstanceProxy::Initialize(UAnimInstance* InAnimInstance)
 
 
 	SnapshotNode.SnapshotName = UAnimSequencerInstance::SequencerPoseName;
-	ClearSequencePlayerMap();
+	ClearSequencePlayerAndMirrorMaps();
 	UpdateCounter.Reset();
 	RootMotionOverride.Reset();
 }
@@ -78,7 +78,7 @@ void FAnimSequencerInstanceProxy::ConstructNodes()
 
 }
 
-void FAnimSequencerInstanceProxy::ClearSequencePlayerMap()
+void FAnimSequencerInstanceProxy::ClearSequencePlayerAndMirrorMaps()
 {
 	for (TPair<uint32, FSequencerPlayerBase*>& Iter : SequencerToPlayerMap)
 	{
@@ -86,6 +86,13 @@ void FAnimSequencerInstanceProxy::ClearSequencePlayerMap()
 	}
 
 	SequencerToPlayerMap.Empty();
+
+	for (TPair<uint32, FAnimNode_Mirror*>& Iter : SequencerToMirrorMap)
+	{
+		delete Iter.Value;
+	}
+
+	SequencerToMirrorMap.Empty();
 }
 
 void FAnimSequencerInstanceProxy::ResetPose()
@@ -101,7 +108,7 @@ void FAnimSequencerInstanceProxy::ResetNodes()
 
 FAnimSequencerInstanceProxy::~FAnimSequencerInstanceProxy()
 {
-	ClearSequencePlayerMap();
+	ClearSequencePlayerAndMirrorMaps();
 }
 
 void FAnimSequencerInstanceProxy::InitAnimTrack(UAnimSequenceBase* InAnimSequence, uint32 SequenceId)
@@ -131,9 +138,15 @@ void FAnimSequencerInstanceProxy::InitAnimTrack(UAnimSequenceBase* InAnimSequenc
 			
 			SequencerToPlayerMap.Add(SequenceId, NewPlayerState);
 
-			// link player to blendnode, this will let you trigger notifies and so on
+			// link player to mirror node,
+			FAnimNode_Mirror* NewMirrorNode = new FAnimNode_Mirror();
+			NewMirrorNode->bMirror = false;
+			NewMirrorNode->Source.SetLinkNode(&NewPlayerState->PlayerNode);
+			SequencerToMirrorMap.Add(SequenceId, NewMirrorNode); 
+
+			// link mirror to blendnode, this will let you trigger notifies and so on
 			NewPlayerState->PlayerNode.SetTeleportToExplicitTime(false);
-			BlendNode.Poses[PoseIndex].SetLinkNode(&NewPlayerState->PlayerNode);
+			BlendNode.Poses[PoseIndex].SetLinkNode(NewMirrorNode);
 
 			// set player state
 			PlayerState = NewPlayerState;
@@ -145,6 +158,13 @@ void FAnimSequencerInstanceProxy::InitAnimTrack(UAnimSequenceBase* InAnimSequenc
 
 		// initialize player
 		PlayerState->PlayerNode.Initialize_AnyThread(FAnimationInitializeContext(this));
+
+		FAnimNode_Mirror* Mirror = SequencerToMirrorMap.FindRef(SequenceId);
+		if (Mirror)
+		{
+			Mirror->Initialize_AnyThread(FAnimationInitializeContext(this));
+			Mirror->CacheBones_AnyThread(FAnimationCacheBonesContext(this));
+		}
 	}
 }
 
@@ -173,21 +193,25 @@ void FAnimSequencerInstanceProxy::TermAnimTrack(int32 SequenceId)
 
 void FAnimSequencerInstanceProxy::UpdateAnimTrack(UAnimSequenceBase* InAnimSequence, uint32 SequenceId, float InPosition, float Weight, bool bFireNotifies)
 {
-	UpdateAnimTrack(InAnimSequence, SequenceId, TOptional<FRootMotionOverride>(), TOptional<float>(), InPosition, Weight, bFireNotifies);
+	UpdateAnimTrack(InAnimSequence, SequenceId, TOptional<FRootMotionOverride>(), TOptional<float>(), InPosition, Weight, bFireNotifies, nullptr);
 }
 
 void FAnimSequencerInstanceProxy::UpdateAnimTrack(UAnimSequenceBase* InAnimSequence, uint32 SequenceId, TOptional<float> InFromPosition, float InToPosition, float Weight, bool bFireNotifies)
 {
-	UpdateAnimTrack(InAnimSequence, SequenceId, TOptional<FRootMotionOverride>(), InFromPosition, InToPosition, Weight, bFireNotifies);
+	UpdateAnimTrack(InAnimSequence, SequenceId, TOptional<FRootMotionOverride>(), InFromPosition, InToPosition, Weight, bFireNotifies, nullptr);
 }
 
-
 void FAnimSequencerInstanceProxy::UpdateAnimTrackWithRootMotion(UAnimSequenceBase* InAnimSequence, int32 SequenceId, const TOptional<FRootMotionOverride>& RootMotion, float InFromPosition, float InToPosition, float Weight, bool bFireNotifies)
 {
-	UpdateAnimTrack(InAnimSequence, SequenceId, RootMotion, InFromPosition, InToPosition, Weight, bFireNotifies);
+	UpdateAnimTrack(InAnimSequence, SequenceId, RootMotion, InFromPosition, InToPosition, Weight, bFireNotifies, nullptr);
+}
 
+void FAnimSequencerInstanceProxy::UpdateAnimTrackWithRootMotion(UAnimSequenceBase* InAnimSequence, int32 SequenceId, const TOptional<FRootMotionOverride>& RootMotion, float InFromPosition, float InToPosition, float Weight, bool bFireNotifies, UMirrorDataTable* InMirrorDataTable)
+{
+	UpdateAnimTrack(InAnimSequence, SequenceId, RootMotion, InFromPosition, InToPosition, Weight, bFireNotifies, InMirrorDataTable);
 }
-void FAnimSequencerInstanceProxy::UpdateAnimTrack(UAnimSequenceBase* InAnimSequence, uint32 SequenceId, const TOptional<FRootMotionOverride>& InRootMotionOverride, TOptional<float> InFromPosition, float InToPosition, float Weight, bool bFireNotifies)
+
+void FAnimSequencerInstanceProxy::UpdateAnimTrack(UAnimSequenceBase* InAnimSequence, uint32 SequenceId, const TOptional<FRootMotionOverride>& InRootMotionOverride, TOptional<float> InFromPosition, float InToPosition, float Weight, bool bFireNotifies, UMirrorDataTable* InMirrorDataTable)
 {
 	EnsureAnimTrack(InAnimSequence, SequenceId);
 
@@ -202,6 +226,20 @@ void FAnimSequencerInstanceProxy::UpdateAnimTrack(UAnimSequenceBase* InAnimSeque
 		PlayerState->PlayerNode.SetExplicitPreviousTime(InFromPosition.GetValue());
 		PlayerState->PlayerNode.SetReinitializationBehavior(ESequenceEvalReinit::NoReset);
 	}
+
+	FAnimNode_Mirror* MirrorNode = SequencerToMirrorMap.FindRef(SequenceId);
+	if (MirrorNode)
+	{
+		MirrorNode->bMirror = InMirrorDataTable != nullptr;
+		UMirrorDataTable* OldMirrorDataTable = MirrorNode->GetMirrorDataTable();
+		MirrorNode->SetMirrorDataTable(InMirrorDataTable);
+
+		if (InMirrorDataTable && OldMirrorDataTable != InMirrorDataTable)
+		{
+			MirrorNode->CacheBones_AnyThread(FAnimationCacheBonesContext(this));
+		}
+	}
+
 	// if no fire notifies, we can teleport to explicit time
 	PlayerState->PlayerNode.SetTeleportToExplicitTime(!bFireNotifies);
 	// if moving to 0.f, we mark this to teleport. Otherwise, do not use explicit time
diff --git a/Engine/Source/Runtime/AnimGraphRuntime/Public/AnimNodes/AnimNode_Mirror.h b/Engine/Source/Runtime/AnimGraphRuntime/Public/AnimNodes/AnimNode_Mirror.h
--- a/Engine/Source/Runtime/AnimGraphRuntime/Public/AnimNodes/AnimNode_Mirror.h
+++ b/Engine/Source/Runtime/AnimGraphRuntime/Public/AnimNodes/AnimNode_Mirror.h
@@ -27,11 +27,14 @@ public:
 
 	float GetBlendTimeOnMirrorStateChange() const { return BlendTimeOnMirrorStateChange; }
 
-#if WITH_EDITOR
 	UMirrorDataTable* GetMirrorDataTable() const;
 	void SetMirrorDataTable(UMirrorDataTable* MirrorTable);
-#endif 
 
+	UPROPERTY(EditAnywhere, Category = Links)
+	FPoseLink Source;
+
+	UPROPERTY(EditAnywhere, Category = Settings, meta = (PinShownByDefault))
+	bool bMirror;
 private:
 
 	UPROPERTY(EditAnywhere, Category = Settings)
@@ -40,12 +43,6 @@ private:
 	UPROPERTY(EditAnywhere, Category = Settings, meta = (PinHiddenByDefault))
 	float BlendTimeOnMirrorStateChange;
 
-	UPROPERTY(EditAnywhere, Category = Links)
-	FPoseLink Source;
-
-	UPROPERTY(EditAnywhere,  Category = Settings, meta = (PinShownByDefault))
-	bool bMirror;
-
 	UPROPERTY(EditAnywhere, Category = MirroredChannels, meta=(DisplayName="Bone"))
 	bool bBoneMirroring;
 
@@ -63,6 +60,7 @@ private:
 	bool bMirrorState;
 	bool bMirrorStateIsValid;
 
+	void FillCompactPoseAndComponentRefRotations(const FBoneContainer& BoneContainer);
 	// Compact pose format of Mirror Bone Map
 	TArray<FCompactPoseBoneIndex> CompactPoseMirrorBones;
 
diff --git a/Engine/Source/Runtime/AnimGraphRuntime/Public/AnimSequencerInstance.h b/Engine/Source/Runtime/AnimGraphRuntime/Public/AnimSequencerInstance.h
--- a/Engine/Source/Runtime/AnimGraphRuntime/Public/AnimSequencerInstance.h
+++ b/Engine/Source/Runtime/AnimGraphRuntime/Public/AnimSequencerInstance.h
@@ -25,7 +25,10 @@ public:
 	virtual void UpdateAnimTrack(UAnimSequenceBase* InAnimSequence, int32 SequenceId, float InFromPosition, float InToPosition, float Weight, bool bFireNotifies);
 
 	/** Update with Root Motion*/
-	void UpdateAnimTrackWithRootMotion(UAnimSequenceBase* InAnimSequence, int32 SequenceId,const TOptional<FRootMotionOverride>& RootMotion, float InFromPosition, float InToPosition, float Weight, bool bFireNotifies);
+	UE_DEPRECATED(5.0, "Please use UpdateAnimTrackWithRootMotion that takes a MirrorDataTable")
+	void UpdateAnimTrackWithRootMotion(UAnimSequenceBase* InAnimSequence, int32 SequenceId, const TOptional<FRootMotionOverride>& RootMotion, float InFromPosition, float InToPosition, float Weight, bool bFireNotifies);
+
+	void UpdateAnimTrackWithRootMotion(UAnimSequenceBase* InAnimSequence, int32 SequenceId,const TOptional<FRootMotionOverride>& RootMotion, float InFromPosition, float InToPosition, float Weight, bool bFireNotifies, UMirrorDataTable* InMirrorDataTable);
 
 	/** Construct all nodes in this instance */
 	virtual void ConstructNodes() override;
diff --git a/Engine/Source/Runtime/AnimGraphRuntime/Public/AnimSequencerInstanceProxy.h b/Engine/Source/Runtime/AnimGraphRuntime/Public/AnimSequencerInstanceProxy.h
--- a/Engine/Source/Runtime/AnimGraphRuntime/Public/AnimSequencerInstanceProxy.h
+++ b/Engine/Source/Runtime/AnimGraphRuntime/Public/AnimSequencerInstanceProxy.h
@@ -7,6 +7,7 @@
 #include "AnimNodes/AnimNode_ApplyAdditive.h"
 #include "AnimNodes/AnimNode_MultiWayBlend.h"
 #include "AnimNodes/AnimNode_PoseSnapshot.h"
+#include "AnimNodes/AnimNode_Mirror.h"
 #include "AnimSequencerInstanceProxy.generated.h"
 
 /** Base class for all 'players' that can attach to and be blended into a sequencer instance's output */
@@ -95,8 +96,11 @@ public:
 	/** Update an animation sequence player in this instance */
 	void UpdateAnimTrack(UAnimSequenceBase* InAnimSequence, uint32 SequenceId, float InPosition, float Weight, bool bFireNotifies);
 	void UpdateAnimTrack(UAnimSequenceBase* InAnimSequence, uint32 SequenceId, TOptional<float> InFromPosition, float InToPosition, float Weight, bool bFireNotifies);
-
+	
+	UE_DEPRECATED(5.0, "Please use the UpdateAnimTrackWithRootMotion that takes a MirrorDataTable")
 	void UpdateAnimTrackWithRootMotion(UAnimSequenceBase* InAnimSequence, int32 SequenceId, const TOptional<FRootMotionOverride>& RootMotion, float InFromPosition, float InToPosition, float Weight, bool bFireNotifies);
+	
+	void UpdateAnimTrackWithRootMotion(UAnimSequenceBase* InAnimSequence, int32 SequenceId, const TOptional<FRootMotionOverride>& RootMotion, float InFromPosition, float InToPosition, float Weight, bool bFireNotifies, UMirrorDataTable* InMirrorDataTable);
 
 	/** Reset all nodes in this instance */
 	virtual void ResetNodes();
@@ -109,7 +113,7 @@ public:
 
 protected:
 
-	void UpdateAnimTrack(UAnimSequenceBase* InAnimSequence, uint32 SequenceId, const TOptional<FRootMotionOverride>& RootMomtionOverride, TOptional<float> InFromPosition, float InToPosition, float Weight, bool bFireNotifies);
+	void UpdateAnimTrack(UAnimSequenceBase* InAnimSequence, uint32 SequenceId, const TOptional<FRootMotionOverride>& RootMomtionOverride, TOptional<float> InFromPosition, float InToPosition, float Weight, bool bFireNotifies, UMirrorDataTable* InMirrorDataTable);
 
 	/** Find a player of a specified type */
 	template<typename Type>
@@ -133,10 +137,13 @@ protected:
 	/** mapping from sequencer index to internal player index */
 	TMap<uint32, FSequencerPlayerBase*> SequencerToPlayerMap;
 
+	/** mapping from sequencer index to internal mirror node index */
+	TMap<uint32, FAnimNode_Mirror*> SequencerToMirrorMap;
+
 	/** custom root motion override sent in from sequencer */
 	TOptional<FRootMotionOverride> RootMotionOverride;
 
 	void InitAnimTrack(UAnimSequenceBase* InAnimSequence, uint32 SequenceId);
 	void EnsureAnimTrack(UAnimSequenceBase* InAnimSequence, uint32 SequenceId);
-	void ClearSequencePlayerMap();
+	void ClearSequencePlayerAndMirrorMaps();
 };
diff --git a/Engine/Source/Runtime/MovieSceneTracks/Private/Evaluation/MovieSceneSkeletalAnimationTemplate.cpp b/Engine/Source/Runtime/MovieSceneTracks/Private/Evaluation/MovieSceneSkeletalAnimationTemplate.cpp
--- a/Engine/Source/Runtime/MovieSceneTracks/Private/Evaluation/MovieSceneSkeletalAnimationTemplate.cpp
+++ b/Engine/Source/Runtime/MovieSceneTracks/Private/Evaluation/MovieSceneSkeletalAnimationTemplate.cpp
@@ -459,8 +459,7 @@ namespace MovieScene
 					RootMotion.GetValue().bBlendFirstChildOfRoot = bBlendFirstChildOfRoot;
 				}
 
-				SequencerInst->UpdateAnimTrackWithRootMotion(InAnimSequence, GetTypeHash(AnimTypeID),RootMotion, InFromPosition, InToPosition, Weight, bFireNotifies);
-
+				SequencerInst->UpdateAnimTrackWithRootMotion(InAnimSequence, GetTypeHash(AnimTypeID),RootMotion, InFromPosition, InToPosition, Weight, bFireNotifies, AnimSection->Params.MirrorDataTable.Get());
 			}
 			else if (UAnimInstance* AnimInst = GetSourceAnimInstance(SkeletalMeshComponent))
 			{
@@ -519,8 +518,8 @@ namespace MovieScene
 					RootMotion.GetValue().RootMotion = RootMotionTransform.GetValue();
 					RootMotion.GetValue().bBlendFirstChildOfRoot = bBlendFirstChildOfRoot;
 				}
-				SequencerInst->UpdateAnimTrackWithRootMotion(InAnimSequence, GetTypeHash(AnimTypeID), RootMotion, InFromPosition, InToPosition, Weight, bFireNotifies);
 
+				SequencerInst->UpdateAnimTrackWithRootMotion(InAnimSequence, GetTypeHash(AnimTypeID), RootMotion, InFromPosition, InToPosition, Weight, bFireNotifies, AnimSection->Params.MirrorDataTable.Get());
 			}
 			else if (UAnimInstance* AnimInst = GetSourceAnimInstance(SkeletalMeshComponent))
 			{
diff --git a/Engine/Source/Runtime/MovieSceneTracks/Private/Sections/MovieSceneSkeletalAnimationSection.cpp b/Engine/Source/Runtime/MovieSceneTracks/Private/Sections/MovieSceneSkeletalAnimationSection.cpp
--- a/Engine/Source/Runtime/MovieSceneTracks/Private/Sections/MovieSceneSkeletalAnimationSection.cpp
+++ b/Engine/Source/Runtime/MovieSceneTracks/Private/Sections/MovieSceneSkeletalAnimationSection.cpp
@@ -24,6 +24,7 @@ namespace
 FMovieSceneSkeletalAnimationParams::FMovieSceneSkeletalAnimationParams()
 {
 	Animation = nullptr;
+	MirrorDataTable = nullptr; 
 	StartOffset_DEPRECATED = SkeletalDeprecatedMagicNumber;
 	EndOffset_DEPRECATED = SkeletalDeprecatedMagicNumber;
 	PlayRate = 1.f;
diff --git a/Engine/Source/Runtime/MovieSceneTracks/Public/Sections/MovieSceneSkeletalAnimationSection.h b/Engine/Source/Runtime/MovieSceneTracks/Public/Sections/MovieSceneSkeletalAnimationSection.h
--- a/Engine/Source/Runtime/MovieSceneTracks/Public/Sections/MovieSceneSkeletalAnimationSection.h
+++ b/Engine/Source/Runtime/MovieSceneTracks/Public/Sections/MovieSceneSkeletalAnimationSection.h
@@ -11,6 +11,7 @@
 
 struct FMovieSceneSkeletalAnimRootMotionTrackParams;
 struct FAnimationPoseData;
+class UMirrorDataTable;
 
 USTRUCT(BlueprintType)
 struct FMovieSceneSkeletalAnimationParams
@@ -53,6 +54,9 @@ struct FMovieSceneSkeletalAnimationParams
 	UPROPERTY(BlueprintReadWrite, EditAnywhere, Category = "Animation" )
 	FName SlotName;
 
+	UPROPERTY(BlueprintReadWrite, EditAnywhere, Category = "Animation")
+	TObjectPtr<class UMirrorDataTable> MirrorDataTable;
+
 	/** The weight curve for this animation section */
 	UPROPERTY()
 	FMovieSceneFloatChannel Weight;
