Fix crashes, inconsistencies and performance issues with layered blend per bone node

Introduces a serial number to FBoneContainer that allows easier refreshing of bone-relative data. Eventually the overall CacheBones path could be superseded with lazy-update systems based on this serial number.

diff --git a/Engine/Source/Editor/AnimGraph/Private/AnimGraphNode_LayeredBoneBlend.cpp b/Engine/Source/Editor/AnimGraph/Private/AnimGraphNode_LayeredBoneBlend.cpp
--- a/Engine/Source/Editor/AnimGraph/Private/AnimGraphNode_LayeredBoneBlend.cpp
+++ b/Engine/Source/Editor/AnimGraph/Private/AnimGraphNode_LayeredBoneBlend.cpp
@@ -43,7 +43,6 @@ void UAnimGraphNode_LayeredBoneBlend::PostEditChangeProperty(struct FPropertyCha
 	if (PropertyName == GET_MEMBER_NAME_STRING_CHECKED(FAnimNode_LayeredBoneBlend, BlendMode))
 	{
 		// If we  change blend modes, we need to resize our containers
-		Node.ValidateData();
 		FBlueprintEditorUtils::MarkBlueprintAsStructurallyModified(GetBlueprint());
 	}
 
@@ -141,12 +140,6 @@ void UAnimGraphNode_LayeredBoneBlend::GetNodeContextMenuActions(UToolMenu* Menu,
 	}
 }
 
-void UAnimGraphNode_LayeredBoneBlend::Serialize(FArchive& Ar)
-{
-	Super::Serialize(Ar);
-	Node.ValidateData();
-}
-
 void UAnimGraphNode_LayeredBoneBlend::ValidateAnimNodeDuringCompilation(class USkeleton* ForSkeleton, class FCompilerResultsLog& MessageLog)
 {
 	UAnimGraphNode_Base::ValidateAnimNodeDuringCompilation(ForSkeleton, MessageLog);
@@ -179,10 +172,10 @@ void UAnimGraphNode_LayeredBoneBlend::ValidateAnimNodeDuringCompilation(class US
 		return;
 	}
 
-	// ensure to cache the data
- 	if (Node.IsCacheInvalid(ForSkeleton))
+	// ensure to cache the per-bone blend weights
+ 	if (!Node.ArePerBoneBlendWeightsValid(ForSkeleton))
  	{
- 		Node.RebuildCacheData(ForSkeleton);
+ 		Node.RebuildPerBoneBlendWeights(ForSkeleton);
  	}
 }
 #undef LOCTEXT_NAMESPACE
diff --git a/Engine/Source/Editor/AnimGraph/Public/AnimGraphNode_LayeredBoneBlend.h b/Engine/Source/Editor/AnimGraph/Public/AnimGraphNode_LayeredBoneBlend.h
--- a/Engine/Source/Editor/AnimGraph/Public/AnimGraphNode_LayeredBoneBlend.h
+++ b/Engine/Source/Editor/AnimGraph/Public/AnimGraphNode_LayeredBoneBlend.h
@@ -16,10 +16,6 @@ class UAnimGraphNode_LayeredBoneBlend : public UAnimGraphNode_BlendListBase
 	UPROPERTY(EditAnywhere, Category=Settings)
 	FAnimNode_LayeredBoneBlend Node;
 
-	// UObject interface
-	virtual void Serialize(FArchive& Ar) override;
-	// End of UObject interface
-
 	// Adds a new pose pin
 	//@TODO: Generalize this behavior (returning a list of actions/delegates maybe?)
 	ANIMGRAPH_API virtual void AddPinToBlendByFilter();
diff --git a/Engine/Source/Runtime/AnimGraphRuntime/Private/AnimNodes/AnimNode_LayeredBoneBlend.cpp b/Engine/Source/Runtime/AnimGraphRuntime/Private/AnimNodes/AnimNode_LayeredBoneBlend.cpp
--- a/Engine/Source/Runtime/AnimGraphRuntime/Private/AnimNodes/AnimNode_LayeredBoneBlend.cpp
+++ b/Engine/Source/Runtime/AnimGraphRuntime/Private/AnimNodes/AnimNode_LayeredBoneBlend.cpp
@@ -29,15 +29,7 @@ void FAnimNode_LayeredBoneBlend::Initialize_AnyThread(const FAnimationInitialize
 	}
 }
 
-#if WITH_EDITOR
-void FAnimNode_LayeredBoneBlend::PostCompile(const class USkeleton* InSkeleton)
-{
-	FAnimNode_Base::PostCompile(InSkeleton);
-	RebuildCacheData(InSkeleton);
-}
-#endif // WITH_EDITOR
-
-void FAnimNode_LayeredBoneBlend::RebuildCacheData(const USkeleton* InSkeleton)
+void FAnimNode_LayeredBoneBlend::RebuildPerBoneBlendWeights(const USkeleton* InSkeleton)
 {
 	if (InSkeleton)
 	{
@@ -55,16 +47,21 @@ void FAnimNode_LayeredBoneBlend::RebuildCacheData(const USkeleton* InSkeleton)
 	}
 }
 
-bool FAnimNode_LayeredBoneBlend::IsCacheInvalid(const USkeleton* InSkeleton) const
+bool FAnimNode_LayeredBoneBlend::ArePerBoneBlendWeightsValid(const USkeleton* InSkeleton) const
 {
-	return (InSkeleton->GetGuid() != SkeletonGuid || InSkeleton->GetVirtualBoneGuid() != VirtualBoneGuid);
+	return (InSkeleton->GetGuid() == SkeletonGuid && InSkeleton->GetVirtualBoneGuid() == VirtualBoneGuid);
 }
 
-void FAnimNode_LayeredBoneBlend::ReinitializeBoneBlendWeights(const FBoneContainer& RequiredBones, const USkeleton* Skeleton)
+void FAnimNode_LayeredBoneBlend::UpdateCachedBoneData(const FBoneContainer& RequiredBones, const USkeleton* Skeleton)
 {
-	if (IsCacheInvalid(Skeleton))
+	if(RequiredBones.GetSerialNumber() == RequiredBonesSerialNumber)
+	{
+		return;
+	}
+
+	if (!ArePerBoneBlendWeightsValid(Skeleton))
 	{
-		RebuildCacheData(Skeleton);
+		RebuildPerBoneBlendWeights(Skeleton);
 	}
 	
 	// build desired bone weights
@@ -126,6 +123,8 @@ void FAnimNode_LayeredBoneBlend::ReinitializeBoneBlendWeights(const FBoneContain
 	{
 		CurvePoseSourceIndices.Reset();
 	}
+
+	RequiredBonesSerialNumber = RequiredBones.GetSerialNumber();
 }
 
 void FAnimNode_LayeredBoneBlend::CacheBones_AnyThread(const FAnimationCacheBonesContext& Context)
@@ -138,10 +137,7 @@ void FAnimNode_LayeredBoneBlend::CacheBones_AnyThread(const FAnimationCacheBones
 		BlendPoses[ChildIndex].CacheBones(Context);
 	}
 
-	if (NumPoses > 0)
-	{
-		ReinitializeBoneBlendWeights(Context.AnimInstanceProxy->GetRequiredBones(), Context.AnimInstanceProxy->GetSkeleton());
-	}
+	UpdateCachedBoneData(Context.AnimInstanceProxy->GetRequiredBones(), Context.AnimInstanceProxy->GetSkeleton());
 }
 
 void FAnimNode_LayeredBoneBlend::Update_AnyThread(const FAnimationUpdateContext& Context)
@@ -163,23 +159,11 @@ void FAnimNode_LayeredBoneBlend::Update_AnyThread(const FAnimationUpdateContext&
 			{
 				if (bHasRelevantPoses == false)
 				{
-					// If our cache is invalid, attempt to update it.
-					if (IsCacheInvalid(Context.AnimInstanceProxy->GetSkeleton()))
-					{
-						ReinitializeBoneBlendWeights(Context.AnimInstanceProxy->GetRequiredBones(), Context.AnimInstanceProxy->GetSkeleton());
-
-						// If Cache is still invalid, we don't have correct DesiredBoneBlendWeights, so abort.
-						// bHasRelevantPoses == false, will passthrough in evaluate.
-						if (!ensure(IsCacheInvalid(Context.AnimInstanceProxy->GetSkeleton())))
-						{
-							break;
-						}
-					}
-					else
-					{
-						FAnimationRuntime::UpdateDesiredBoneWeight(DesiredBoneBlendWeights, CurrentBoneBlendWeights, BlendWeights);
-					}
+					// Update cached data now we know we might be valid
+					UpdateCachedBoneData(Context.AnimInstanceProxy->GetRequiredBones(), Context.AnimInstanceProxy->GetSkeleton());
 
+					// Update weights
+					FAnimationRuntime::UpdateDesiredBoneWeight(DesiredBoneBlendWeights, CurrentBoneBlendWeights, BlendWeights);
 					bHasRelevantPoses = true;
 
 					if(bBlendRootMotionBasedOnRootBone)
@@ -321,59 +305,3 @@ void FAnimNode_LayeredBoneBlend::GatherDebugData(FNodeDebugData& DebugData)
 		BlendPoses[ChildIndex].GatherDebugData(DebugData.BranchFlow(BlendWeights[ChildIndex]));
 	}
 }
-
-#if WITH_EDITOR
-void FAnimNode_LayeredBoneBlend::ValidateData()
-{
-	// ideally you don't like to get to situation where it becomes inconsistent, but this happened, 
-	// and we don't know what caused this. Possibly copy/paste, but I tried copy/paste and that didn't work
-	// so here we add code to fix this up manually in editor, so that they can continue working on it. 
-	int32 PoseNum = BlendPoses.Num();
-	int32 WeightNum = BlendWeights.Num();
-	int32 LayerNum = (BlendMode == ELayeredBoneBlendMode::BranchFilter) ? LayerSetup.Num() : BlendMasks.Num();
-
-	int32 Max = FMath::Max3(PoseNum, WeightNum, LayerNum);
-	int32 Min = FMath::Min3(PoseNum, WeightNum, LayerNum);
-	// if they are not all same
-	if (Min != Max)
-	{
-		// we'd like to increase to all Max
-		// sadly we don't have add X for how many
-		for (int32 Index=PoseNum; Index<Max; ++Index)
-		{
-			BlendPoses.Add(FPoseLink());
-		}
-
-		for(int32 Index=WeightNum; Index<Max; ++Index)
-		{
-			BlendWeights.Add(1.f);
-		}
-
-		if (BlendMode == ELayeredBoneBlendMode::BranchFilter)
-		{
-			for (int32 Index = LayerNum; Index < Max; ++Index)
-			{
-				LayerSetup.Add(FInputBlendPose());
-			}
-		}
-		else
-		{
-			// Get rid of our branch filters if we are not using them to avoid saving into the asset
-			LayerSetup.Reset();
-		}
-
-		if (BlendMode == ELayeredBoneBlendMode::BlendMask)
-		{
-			for (int32 Index = LayerNum; Index < Max; ++Index)
-			{
-				BlendMasks.Add(nullptr);
-			}
-		}
-		else
-		{
-			// Get rid of blend masks array if we are not using them to avoid saving into the asset
-			BlendMasks.Reset();
-		}
-	}
-}
-#endif
diff --git a/Engine/Source/Runtime/AnimGraphRuntime/Public/AnimNodes/AnimNode_LayeredBoneBlend.h b/Engine/Source/Runtime/AnimGraphRuntime/Public/AnimNodes/AnimNode_LayeredBoneBlend.h
--- a/Engine/Source/Runtime/AnimGraphRuntime/Public/AnimNodes/AnimNode_LayeredBoneBlend.h
+++ b/Engine/Source/Runtime/AnimGraphRuntime/Public/AnimNodes/AnimNode_LayeredBoneBlend.h
@@ -82,19 +82,18 @@ public:
 	int32 LODThreshold;
 
 protected:
-
-	// This is buffer to serialize blend weight data for each joints
-	// This has to save with the corresponding SkeletopnGuid
-	// If not, it will rebuild in run-time
+	// Per-bone weights for the skeleton. Serialized as these are only relative to the skeleton, but can potentially
+	// be regenerated at runtime if the GUIDs dont match
 	UPROPERTY()
 	TArray<FPerBoneBlendWeight>	PerBoneBlendWeights;
 
+	// Guids for skeleton used to determine whether the PerBoneBlendWeights need rebuilding
 	UPROPERTY()
-	FGuid						SkeletonGuid;
+	FGuid SkeletonGuid;
 
+	// Guid for virtual bones used to determine whether the PerBoneBlendWeights need rebuilding
 	UPROPERTY()
-	FGuid						VirtualBoneGuid;
-
+	FGuid VirtualBoneGuid;
 
 	// transient data to handle weight and target weight
 	// this array changes based on required bones
@@ -102,6 +101,9 @@ protected:
 	TArray<FPerBoneBlendWeight> CurrentBoneBlendWeights;
 	TArray<uint8> CurvePoseSourceIndices;
 
+	// Serial number of the required bones container
+	uint16 RequiredBonesSerialNumber;
+	
 public:	
 	FAnimNode_LayeredBoneBlend()
 		: BlendMode(ELayeredBoneBlendMode::BranchFilter)
@@ -111,6 +113,7 @@ public:
 		, bBlendRootMotionBasedOnRootBone(true)
 		, bHasRelevantPoses(false)
 		, LODThreshold(INDEX_NONE)
+		, RequiredBonesSerialNumber(0)
 	{
 	}
 
@@ -137,23 +140,15 @@ public:
 		LayerSetup.RemoveAt(PoseIndex);
 	}
 
-#if WITH_EDITOR
-	// ideally you don't like to get to situation where it becomes inconsistent, but this happened, 
-	// and we don't know what caused this. Possibly copy/paste, but I tried copy/paste and that didn't work
-	// so here we add code to fix this up manually in editor, so that they can continue working on it. 
-	void ValidateData();
-	// FAnimNode_Base interface
-	virtual void PostCompile(const class USkeleton* InSkeleton) override;
-	// end FAnimNode_Base interface
-#endif
+private:
+	// Rebuild cache per bone blend weights from the skeleton
+	void RebuildPerBoneBlendWeights(const USkeleton* InSkeleton);
 
-	/** Reinitialize bone weights */
-	void ReinitializeBoneBlendWeights(const FBoneContainer& RequiredBones, const USkeleton* Skeleton);
+	// Check whether per-bone blend weights are valid according to the skeleton (GUID check)
+	bool ArePerBoneBlendWeightsValid(const USkeleton* InSkeleton) const;
 
-private:
-	// Rebuild cache data from the skeleton
-	void RebuildCacheData(const USkeleton* InSkeleton);
-	bool IsCacheInvalid(const USkeleton* InSkeleton) const;
+	// Update cached data if required
+	void UpdateCachedBoneData(const FBoneContainer& RequiredBones, const USkeleton* Skeleton);
 
 	friend class UAnimGraphNode_LayeredBoneBlend;
 };
diff --git a/Engine/Source/Runtime/Engine/Private/BoneContainer.cpp b/Engine/Source/Runtime/Engine/Private/BoneContainer.cpp
--- a/Engine/Source/Runtime/Engine/Private/BoneContainer.cpp
+++ b/Engine/Source/Runtime/Engine/Private/BoneContainer.cpp
@@ -55,6 +55,7 @@ FBoneContainer::FBoneContainer()
 , AssetSkeleton(nullptr)
 , RefSkeleton(nullptr)
 , UIDToArrayIndexLUTValidCount(0)
+, SerialNumber(0)
 #if DO_CHECK
 , CalculatedForLOD(INDEX_NONE)
 #endif
@@ -75,6 +76,7 @@ FBoneContainer::FBoneContainer(const TArray<FBoneIndexType>& InRequiredBoneIndex
 , AssetSkeleton(nullptr)
 , RefSkeleton(nullptr)
 , UIDToArrayIndexLUTValidCount(0)
+, SerialNumber(0)
 #if DO_CHECK
 , CalculatedForLOD(INDEX_NONE)
 #endif
@@ -236,6 +238,8 @@ void FBoneContainer::Initialize(const FCurveEvaluationOption& CurveEvalOption)
 
 	// Reset retargeting cached data look up table.
 	RetargetSourceCachedDataLUT.Reset();
+
+	RegenerateSerialNumber();
 }
 
 void FBoneContainer::CacheRequiredAnimCurveUids(const FCurveEvaluationOption& CurveEvalOption)
@@ -349,6 +353,20 @@ void FBoneContainer::CacheRequiredAnimCurveUids(const FCurveEvaluationOption& Cu
 
 	// Make sure we regenerate our cached curve mappings next time they are requested.
 	MarkAllCachedCurveMappingsDirty();
+
+	RegenerateSerialNumber();
+}
+
+void FBoneContainer::RegenerateSerialNumber()
+{
+	// Bump the serial number
+	SerialNumber++;
+
+	// Skip zero as this is used to indicate an invalid bone container
+	if(SerialNumber == 0)
+	{
+		SerialNumber++;
+	}
 }
 
 void FBoneContainer::MarkAllCachedCurveMappingsDirty()
diff --git a/Engine/Source/Runtime/Engine/Public/BoneContainer.h b/Engine/Source/Runtime/Engine/Public/BoneContainer.h
--- a/Engine/Source/Runtime/Engine/Public/BoneContainer.h
+++ b/Engine/Source/Runtime/Engine/Public/BoneContainer.h
@@ -213,6 +213,11 @@ private:
 	/** Look up table of UID to FAnimCurveType UIDToNameLUT[InUID] = FAnimCurveType of curve. */
 	TArray<FAnimCurveType> UIDToCurveTypeLUT;
 
+	// The serial number of this bone container. This is incremented each time the container is regenerated and can
+	// be used to track whether to cache bone data. If this value is zero then the bone container is considered invalid
+	// as it has never been regenerated.
+	uint16 SerialNumber;
+
 	/** For debugging. */
 #if DO_CHECK
 	/** The LOD that we calculated required bones when regenerated */
@@ -494,6 +499,9 @@ public:
 	const FCachedSkeletonCurveMapping& GetOrCreateCachedCurveMapping(const FSkeletonRemapping* SkeletonRemapping) const;
 	void MarkAllCachedCurveMappingsDirty();
 
+	// Get the serial number of this bone container. @see SerialNumber
+	uint16 GetSerialNumber() const { return SerialNumber; }
+	
 private:
 	/** The map of cached curve mapping indexes, which is used for skeleton remapping. The key of this table is the source skeleton of the asset we are sampling curve values from. */
 	mutable TMap<USkeleton*, FCachedSkeletonCurveMapping> CachedCurveMappingTable;
@@ -515,6 +523,9 @@ private:
 
 	/** Cache remapping data if current Asset is a Skeleton, with all compatible Skeletons. */
 	void RemapFromSkeleton(USkeleton const & SourceSkeleton);
+
+	// Regenerate the serial number after internal data is updated
+	void RegenerateSerialNumber();
 };
 
 
