Added new functionality to ModifyCurve anim node that allows curve manipulation via a name-value map

diff --git a/Engine/Source/Runtime/AnimGraphRuntime/Private/AnimNodes/AnimNode_ModifyCurve.cpp b/Engine/Source/Runtime/AnimGraphRuntime/Private/AnimNodes/AnimNode_ModifyCurve.cpp
--- a/Engine/Source/Runtime/AnimGraphRuntime/Private/AnimNodes/AnimNode_ModifyCurve.cpp
+++ b/Engine/Source/Runtime/AnimGraphRuntime/Private/AnimNodes/AnimNode_ModifyCurve.cpp
@@ -8,6 +8,7 @@ FAnimNode_ModifyCurve::FAnimNode_ModifyCurve()
 {
 	ApplyMode = EModifyCurveApplyMode::Blend;
 	Alpha = 1.f;
+	bInitializeLastValuesMap = false;
 }
 
 void FAnimNode_ModifyCurve::Initialize_AnyThread(const FAnimationInitializeContext& Context)
@@ -21,6 +22,9 @@ void FAnimNode_ModifyCurve::Initialize_AnyThread(const FAnimationInitializeConte
 	{
 		LastCurveValues.Reset(CurveValues.Num());
 		LastCurveValues.AddZeroed(CurveValues.Num());
+
+		// Indicate that last values for curve map should be initialized on evaluate since updated CurveMap is not yet available.
+		bInitializeLastValuesMap = true;		
 	}
 }
 
@@ -42,55 +46,106 @@ void FAnimNode_ModifyCurve::Evaluate_AnyThread(FPoseContext& Output)
 	//	Morph target and Material parameter curves
 	USkeleton* Skeleton = Output.AnimInstanceProxy->GetSkeleton();
 
-	check(CurveNames.Num() == CurveValues.Num());
-
-	for (int32 ModIdx = 0; ModIdx < CurveNames.Num(); ModIdx++)
+	// Check if Last Values Map should be initialized. This has to be done if curve map is changed in the node graph without compiling.
+	if (bInitializeLastValuesMap)
 	{
-		FName CurveName = CurveNames[ModIdx];
-		SmartName::UID_Type NameUID = Skeleton->GetUIDByName(USkeleton::AnimCurveMappingName, CurveName);
+		// CurveMap is initialized so Last curve map values can be stored.
+		LastCurveMapValues.Reset();
+		LastCurveMapValues.Reserve(CurveMap.Num());
+		for (auto It = CurveMap.CreateConstIterator(); It; ++It)
+		{
+			LastCurveMapValues.Add(It.Key(), 0.0f);
+		}
+		bInitializeLastValuesMap = false;
+	}
+
+	// Process curve map if available.
+	for (auto& CurveNameValue : CurveMap)
+	{				
+		SmartName::UID_Type NameUID = Skeleton->GetUIDByName(USkeleton::AnimCurveMappingName, CurveNameValue.Key);
 		if (NameUID != SmartName::MaxUID)
 		{
-			float CurveValue = CurveValues[ModIdx];
 			float CurrentValue = Output.Curve.Get(NameUID);
-
-			// Use ApplyMode enum to decide how to apply
-			if (ApplyMode == EModifyCurveApplyMode::Add)
+			float NewValue = CurveNameValue.Value;
+			
+			if (ApplyMode == EModifyCurveApplyMode::WeightedMovingAverage)
 			{
-				Output.Curve.Set(NameUID, CurrentValue + CurveValue);
+				float* LastValue = LastCurveMapValues.Find(CurveNameValue.Key);
+				if (LastValue)
+				{
+					ProcessCurveWMAOperation(Output, NameUID, CurrentValue, NewValue, *LastValue);	
+				}
 			}
-			else if (ApplyMode == EModifyCurveApplyMode::Scale)
+			else
 			{
-				Output.Curve.Set(NameUID, CurrentValue * CurveValue);
+				ProcessCurveOperation(ApplyMode, Output, NameUID, CurrentValue, NewValue);
 			}
-			else if (ApplyMode == EModifyCurveApplyMode::WeightedMovingAverage)
-			{
-				check(LastCurveValues.Num() == CurveValues.Num());
+		}
+	}
 
-				const float LastCurveValue = LastCurveValues[ModIdx];
-				const float WAvg = FMath::WeightedMovingAverage(CurrentValue, LastCurveValue, Alpha);
-				// Update the last curve value for next run
-				LastCurveValues[ModIdx] = WAvg;
+	// Process custom selected curve pins.
+	check(CurveNames.Num() == CurveValues.Num());
 
-				Output.Curve.Set(NameUID, WAvg);
-			}
-			else if (ApplyMode == EModifyCurveApplyMode::RemapCurve)
+	for (int32 ModIdx = 0; ModIdx < CurveNames.Num(); ModIdx++)
+	{
+		FName CurveName = CurveNames[ModIdx];
+		SmartName::UID_Type NameUID = Skeleton->GetUIDByName(USkeleton::AnimCurveMappingName, CurveName);
+		if (NameUID != SmartName::MaxUID)
+		{
+			float NewValue = CurveValues[ModIdx];
+			float CurrentValue = Output.Curve.Get(NameUID);
+			if (ApplyMode == EModifyCurveApplyMode::WeightedMovingAverage)
 			{
-				const float RemapScale = 1.f / FMath::Max(1.f - CurveValue, 0.01f);
-				const float RemappedValue = FMath::Min(FMath::Max(CurrentValue - CurveValue, 0.f) * RemapScale, 1.f);
-				Output.Curve.Set(NameUID, RemappedValue);
+				check(LastCurveValues.Num() == CurveValues.Num())
+				float& LastValue = LastCurveValues[ModIdx];
+				ProcessCurveWMAOperation(Output, NameUID, CurrentValue, NewValue, LastValue);
 			}
-			else // Blend
+			else
 			{
-				float UseAlpha = FMath::Clamp(Alpha, 0.f, 1.f);
-				Output.Curve.Set(NameUID, FMath::Lerp(CurrentValue, CurveValue, UseAlpha));
+				ProcessCurveOperation(ApplyMode, Output, NameUID, CurrentValue, NewValue);
 			}
 		}
 	}
 }
 
+void FAnimNode_ModifyCurve::ProcessCurveOperation(const EModifyCurveApplyMode& InApplyMode, FPoseContext& Output, const SmartName::UID_Type& NameUID, float CurrentValue, float NewValue)
+{
+	// Use ApplyMode enum to decide how to apply
+	if (InApplyMode == EModifyCurveApplyMode::Add)
+	{
+		Output.Curve.Set(NameUID, CurrentValue + NewValue);
+	}
+	else if (InApplyMode == EModifyCurveApplyMode::Scale)
+	{
+		Output.Curve.Set(NameUID, CurrentValue * NewValue);
+	}
+	else if (InApplyMode == EModifyCurveApplyMode::RemapCurve)
+	{
+		const float RemapScale = 1.f / FMath::Max(1.f - NewValue, 0.01f);
+		const float RemappedValue = FMath::Min(FMath::Max(CurrentValue - NewValue, 0.f) * RemapScale, 1.f);
+		Output.Curve.Set(NameUID, RemappedValue);
+	}
+	else if (InApplyMode == EModifyCurveApplyMode::Blend)
+	{
+		float UseAlpha = FMath::Clamp(Alpha, 0.f, 1.f);
+		Output.Curve.Set(NameUID, FMath::Lerp(CurrentValue, NewValue, UseAlpha));
+	}
+}
+
+
+void FAnimNode_ModifyCurve::ProcessCurveWMAOperation(FPoseContext& Output, const SmartName::UID_Type& NameUID, float CurrentValue, float NewValue, float& InOutLastValue)
+{
+	const float WAvg = FMath::WeightedMovingAverage(CurrentValue, InOutLastValue, Alpha);
+	// Update the last curve value for next run
+	InOutLastValue = WAvg;
+
+	Output.Curve.Set(NameUID, WAvg);
+}
+
 void FAnimNode_ModifyCurve::Update_AnyThread(const FAnimationUpdateContext& Context)
 {
 	DECLARE_SCOPE_HIERARCHICAL_COUNTER_ANIMNODE(Update_AnyThread)
+
 	// Run update on input pose nodes
 	SourcePose.Update(Context);
 
diff --git a/Engine/Source/Runtime/AnimGraphRuntime/Public/AnimNodes/AnimNode_ModifyCurve.h b/Engine/Source/Runtime/AnimGraphRuntime/Public/AnimNodes/AnimNode_ModifyCurve.h
--- a/Engine/Source/Runtime/AnimGraphRuntime/Public/AnimNodes/AnimNode_ModifyCurve.h
+++ b/Engine/Source/Runtime/AnimGraphRuntime/Public/AnimNodes/AnimNode_ModifyCurve.h
@@ -35,6 +35,9 @@ struct ANIMGRAPHRUNTIME_API FAnimNode_ModifyCurve : public FAnimNode_Base
 	UPROPERTY(EditAnywhere, EditFixedSize, BlueprintReadWrite, Category = Links)
 	FPoseLink SourcePose;
 
+	UPROPERTY(EditAnywhere, BlueprintReadWrite, editfixedsize, Category = ModifyCurve, meta = (PinHiddenByDefault))
+	TMap<FName, float> CurveMap;
+
 	UPROPERTY(EditAnywhere, BlueprintReadWrite, editfixedsize, Category = ModifyCurve, meta = (PinShownByDefault))
 	TArray<float> CurveValues;
 
@@ -42,6 +45,8 @@ struct ANIMGRAPHRUNTIME_API FAnimNode_ModifyCurve : public FAnimNode_Base
 	TArray<FName> CurveNames;
 
 	TArray<float> LastCurveValues;
+	TMap<FName, float> LastCurveMapValues;
+	bool bInitializeLastValuesMap;
 
 	UPROPERTY(EditAnywhere, BlueprintReadWrite, Category = ModifyCurve, meta = (PinShownByDefault))
 	float Alpha;
@@ -55,6 +60,7 @@ struct ANIMGRAPHRUNTIME_API FAnimNode_ModifyCurve : public FAnimNode_Base
 	virtual void Initialize_AnyThread(const FAnimationInitializeContext& Context) override;
 	virtual void CacheBones_AnyThread(const FAnimationCacheBonesContext& Context) override;
 	virtual void Evaluate_AnyThread(FPoseContext& Output) override;
+
 	virtual void Update_AnyThread(const FAnimationUpdateContext& Context) override;
 	// End of FAnimNode_Base interface
 
@@ -64,4 +70,8 @@ struct ANIMGRAPHRUNTIME_API FAnimNode_ModifyCurve : public FAnimNode_Base
 	/** Remove a curve from being modified */
 	void RemoveCurve(int32 PoseIndex);
 #endif // WITH_EDITOR
+
+private:
+	void ProcessCurveOperation(const EModifyCurveApplyMode& InApplyMode, FPoseContext& Output, const SmartName::UID_Type& NameUID, float CurrentValue, float NewValue);
+	void ProcessCurveWMAOperation(FPoseContext& Output, const SmartName::UID_Type& NameUID, float CurrentValue, float NewValue, float& InOutLastValue);
 };
