Added more library functions for sequence players/evaluators

Fixed execution context conversion & added BP-exposed conversion funcitons

diff --git a/Engine/Source/Runtime/AnimGraphRuntime/Private/AnimExecutionContextLibrary.cpp b/Engine/Source/Runtime/AnimGraphRuntime/Private/AnimExecutionContextLibrary.cpp
--- a/Engine/Source/Runtime/AnimGraphRuntime/Private/AnimExecutionContextLibrary.cpp
+++ b/Engine/Source/Runtime/AnimGraphRuntime/Private/AnimExecutionContextLibrary.cpp
@@ -19,4 +19,24 @@ FAnimNodeReference UAnimExecutionContextLibrary::GetAnimNodeReference(UAnimInsta
 	// As the index is patched during compilation, it needs to be reversed here
 	int32 ReverseIndex = AnimNodeProperties.Num() - 1 - Index;
 	return FAnimNodeReference(Instance, ReverseIndex);
+}
+
+FAnimInitializationContext UAnimExecutionContextLibrary::ConvertToInitializationContext(const FAnimExecutionContext& Context, EAnimExecutionContextConversionResult& Result)
+{
+	return FAnimExecutionContext::ConvertToType<FAnimInitializationContext>(Context, Result);
+}
+
+FAnimUpdateContext UAnimExecutionContextLibrary::ConvertToUpdateContext(const FAnimExecutionContext& Context, EAnimExecutionContextConversionResult& Result)
+{
+	return FAnimExecutionContext::ConvertToType<FAnimUpdateContext>(Context, Result);
+}
+
+FAnimPoseContext UAnimExecutionContextLibrary::ConvertToPoseContext(const FAnimExecutionContext& Context, EAnimExecutionContextConversionResult& Result)
+{
+	return FAnimExecutionContext::ConvertToType<FAnimPoseContext>(Context, Result);
+}
+
+FAnimComponentSpacePoseContext UAnimExecutionContextLibrary::ConvertToComponentSpacePoseContext(const FAnimExecutionContext& Context, EAnimExecutionContextConversionResult& Result)
+{
+	return FAnimExecutionContext::ConvertToType<FAnimComponentSpacePoseContext>(Context, Result);
 }
\ No newline at end of file
diff --git a/Engine/Source/Runtime/AnimGraphRuntime/Private/AnimNodes/AnimNode_SequenceEvaluator.cpp b/Engine/Source/Runtime/AnimGraphRuntime/Private/AnimNodes/AnimNode_SequenceEvaluator.cpp
--- a/Engine/Source/Runtime/AnimGraphRuntime/Private/AnimNodes/AnimNode_SequenceEvaluator.cpp
+++ b/Engine/Source/Runtime/AnimGraphRuntime/Private/AnimNodes/AnimNode_SequenceEvaluator.cpp
@@ -157,6 +157,21 @@ float FAnimNode_SequenceEvaluator::GetExplicitTime() const
 	return GET_ANIM_NODE_DATA(float, ExplicitTime);
 }
 
+bool FAnimNode_SequenceEvaluator::SetExplicitTime(float InTime)
+{
+#if WITH_EDITORONLY_DATA
+	ExplicitTime = InTime;
+#endif
+
+	if (float* ExplicitTimePtr = GET_INSTANCE_ANIM_NODE_DATA_PTR(float, ExplicitTime))
+	{
+		*ExplicitTimePtr = InTime;
+		return true;
+	}
+
+	return false;
+}
+
 bool FAnimNode_SequenceEvaluator::GetShouldLoop() const
 {
 	return GET_ANIM_NODE_DATA(bool, bShouldLoop);
diff --git a/Engine/Source/Runtime/AnimGraphRuntime/Private/SequenceEvaluatorLibrary.cpp b/Engine/Source/Runtime/AnimGraphRuntime/Private/SequenceEvaluatorLibrary.cpp
new file mode 100644
--- /dev/null
+++ b/Engine/Source/Runtime/AnimGraphRuntime/Private/SequenceEvaluatorLibrary.cpp
@@ -0,0 +1,76 @@
+// Copyright Epic Games, Inc. All Rights Reserved.
+
+#include "SequenceEvaluatorLibrary.h"
+
+#include "Animation/AnimNode_Inertialization.h"
+#include "AnimNodes/AnimNode_SequenceEvaluator.h"
+
+DEFINE_LOG_CATEGORY_STATIC(LogSequenceEvaluatorLibrary, Verbose, All);
+
+FSequenceEvaluatorReference USequenceEvaluatorLibrary::ConvertToSequenceEvaluatorContext(const FAnimNodeReference& Node, EAnimNodeReferenceConversionResult& Result)
+{
+	return FAnimNodeReference::ConvertToType<FSequenceEvaluatorReference>(Node, Result);
+}
+
+FSequenceEvaluatorReference USequenceEvaluatorLibrary::SetExplicitTime(const FSequenceEvaluatorReference& SequenceEvaluator, float Time)
+{
+	SequenceEvaluator.CallAnimNodeFunction<FAnimNode_SequenceEvaluator>(
+		TEXT("SetExplicitTime"),
+		[Time](FAnimNode_SequenceEvaluator& InSequenceEvaluator)
+		{
+			if(!InSequenceEvaluator.SetExplicitTime(Time))
+			{
+				UE_LOG(LogSequenceEvaluatorLibrary, Warning, TEXT("Could not set explicit time on sequence evaluator, value is not dynamic. Set it as Always Dynamic."));
+			}
+		});
+
+	return SequenceEvaluator;
+}
+
+FSequenceEvaluatorReference USequenceEvaluatorLibrary::SetSequence(const FSequenceEvaluatorReference& SequenceEvaluator, UAnimSequenceBase* Sequence)
+{
+	SequenceEvaluator.CallAnimNodeFunction<FAnimNode_SequenceEvaluator>(
+		TEXT("SetSequence"),
+		[Sequence](FAnimNode_SequenceEvaluator& InSequenceEvaluator)
+		{
+			if(!InSequenceEvaluator.SetSequence(Sequence))
+			{
+				UE_LOG(LogSequenceEvaluatorLibrary, Warning, TEXT("Could not set sequence on sequence evaluator, value is not dynamic. Set it as Always Dynamic."));
+			}
+		});
+
+	return SequenceEvaluator;
+}
+
+FSequenceEvaluatorReference USequenceEvaluatorLibrary::SetSequenceWithInertialBlending(const FAnimUpdateContext& UpdateContext, const FSequenceEvaluatorReference& SequenceEvaluator, UAnimSequenceBase* Sequence, float BlendTime)
+{
+	SequenceEvaluator.CallAnimNodeFunction<FAnimNode_SequenceEvaluator>(
+		TEXT("SetSequenceWithInterialBlending"),
+		[Sequence, &UpdateContext, BlendTime](FAnimNode_SequenceEvaluator& InSequenceEvaluator)
+		{
+			const UAnimSequenceBase* CurrentSequence = InSequenceEvaluator.GetSequence();
+			const bool bAnimSequenceChanged = (CurrentSequence != Sequence);
+			
+			if(!InSequenceEvaluator.SetSequence(Sequence))
+			{
+				UE_LOG(LogSequenceEvaluatorLibrary, Warning, TEXT("Could not set sequence on sequence evaluator, value is not dynamic. Set it as Always Dynamic."));
+			}
+
+			if(bAnimSequenceChanged && BlendTime > 0.0f)
+			{
+				if (const FAnimationUpdateContext* AnimationUpdateContext = UpdateContext.GetContext())
+				{
+					if (UE::Anim::IInertializationRequester* InertializationRequester = AnimationUpdateContext->GetMessage<UE::Anim::IInertializationRequester>())
+					{
+						InertializationRequester->RequestInertialization(BlendTime);
+					}
+				}
+				else
+				{
+					UE_LOG(LogSequenceEvaluatorLibrary, Warning, TEXT("SetSequenceWithInterialBlending called with invalid context"));
+				}
+			}
+		});
+
+	return SequenceEvaluator;
+}
diff --git a/Engine/Source/Runtime/AnimGraphRuntime/Private/SequencePlayerLibrary.cpp b/Engine/Source/Runtime/AnimGraphRuntime/Private/SequencePlayerLibrary.cpp
--- a/Engine/Source/Runtime/AnimGraphRuntime/Private/SequencePlayerLibrary.cpp
+++ b/Engine/Source/Runtime/AnimGraphRuntime/Private/SequencePlayerLibrary.cpp
@@ -1,45 +1,104 @@
 // Copyright Epic Games, Inc. All Rights Reserved.
 
 #include "SequencePlayerLibrary.h"
+
+#include "Animation/AnimNode_Inertialization.h"
 #include "Animation/AnimNode_SequencePlayer.h"
 
+DEFINE_LOG_CATEGORY_STATIC(LogSequencePlayerLibrary, Verbose, All);
+
 FSequencePlayerReference USequencePlayerLibrary::ConvertToSequencePlayerContext(const FAnimNodeReference& Node, EAnimNodeReferenceConversionResult& Result)
 {
 	return FAnimNodeReference::ConvertToType<FSequencePlayerReference>(Node, Result);
 }
 
-FSequencePlayerReference USequencePlayerLibrary::SetAccumulatedTime(const FSequencePlayerReference& SequencePlayerContext, float Time)
+FSequencePlayerReference USequencePlayerLibrary::SetAccumulatedTime(const FSequencePlayerReference& SequencePlayer, float Time)
 {
-	SequencePlayerContext.CallAnimNodeFunction<FAnimNode_SequencePlayer>(
+	SequencePlayer.CallAnimNodeFunction<FAnimNode_SequencePlayer>(
 		TEXT("SetAccumulatedTime"),
-		[Time](FAnimNode_SequencePlayer& SequencePlayer)
+		[Time](FAnimNode_SequencePlayer& InSequencePlayer)
 		{
-			SequencePlayer.SetAccumulatedTime(Time);
+			InSequencePlayer.SetAccumulatedTime(Time);
 		});
 
-	return SequencePlayerContext;
+	return SequencePlayer;
 }
 
-FSequencePlayerReference USequencePlayerLibrary::SetStartPosition(const FSequencePlayerReference& SequencePlayerContext, float StartPosition)
+FSequencePlayerReference USequencePlayerLibrary::SetStartPosition(const FSequencePlayerReference& SequencePlayer, float StartPosition)
 {
-	SequencePlayerContext.CallAnimNodeFunction<FAnimNode_SequencePlayer>(
+	SequencePlayer.CallAnimNodeFunction<FAnimNode_SequencePlayer>(
 		TEXT("SetStartPosition"),
-		[StartPosition](FAnimNode_SequencePlayer& SequencePlayer)
+		[StartPosition](FAnimNode_SequencePlayer& InSequencePlayer)
 		{
-			SequencePlayer.SetStartPosition(StartPosition);
+			if(!InSequencePlayer.SetStartPosition(StartPosition))
+			{
+				UE_LOG(LogSequencePlayerLibrary, Warning, TEXT("Could not set start position on sequence player, value is not dynamic. Set it as Always Dynamic."));
+			}
 		});
 
-	return SequencePlayerContext;
+	return SequencePlayer;
 }
 
-FSequencePlayerReference USequencePlayerLibrary::SetPlayRate(const FSequencePlayerReference& SequencePlayerContext, float PlayRate)
+FSequencePlayerReference USequencePlayerLibrary::SetPlayRate(const FSequencePlayerReference& SequencePlayer, float PlayRate)
 {
-	SequencePlayerContext.CallAnimNodeFunction<FAnimNode_SequencePlayer>(
+	SequencePlayer.CallAnimNodeFunction<FAnimNode_SequencePlayer>(
 		TEXT("SetPlayRate"),
-		[PlayRate](FAnimNode_SequencePlayer& SequencePlayer)
+		[PlayRate](FAnimNode_SequencePlayer& InSequencePlayer)
+		{
+			if(!InSequencePlayer.SetPlayRate(PlayRate))
+			{
+				UE_LOG(LogSequencePlayerLibrary, Warning, TEXT("Could not set play rate on sequence player, value is not dynamic. Set it as Always Dynamic."));
+			}
+		});
+
+	return SequencePlayer;
+}
+
+FSequencePlayerReference USequencePlayerLibrary::SetSequence(const FSequencePlayerReference& SequencePlayer, UAnimSequenceBase* Sequence)
+{
+	SequencePlayer.CallAnimNodeFunction<FAnimNode_SequencePlayer>(
+		TEXT("SetSequence"),
+		[Sequence](FAnimNode_SequencePlayer& InSequencePlayer)
 		{
-			SequencePlayer.SetPlayRate(PlayRate);
+			if(!InSequencePlayer.SetSequence(Sequence))
+			{
+				UE_LOG(LogSequencePlayerLibrary, Warning, TEXT("Could not set sequence on sequence player, value is not dynamic. Set it as Always Dynamic."));
+			}
+		});
+
+	return SequencePlayer;
+}
+
+
+FSequencePlayerReference USequencePlayerLibrary::SetSequenceWithInertialBlending(const FAnimUpdateContext& UpdateContext, const FSequencePlayerReference& SequencePlayer, UAnimSequenceBase* Sequence, float BlendTime)
+{
+	SequencePlayer.CallAnimNodeFunction<FAnimNode_SequencePlayer>(
+		TEXT("SetSequenceWithInterialBlending"),
+		[Sequence, &UpdateContext, BlendTime](FAnimNode_SequencePlayer& InSequencePlayer)
+		{
+			const UAnimSequenceBase* CurrentSequence = InSequencePlayer.GetSequence();
+			const bool bAnimSequenceChanged = (CurrentSequence != Sequence);
+			
+			if(!InSequencePlayer.SetSequence(Sequence))
+			{
+				UE_LOG(LogSequencePlayerLibrary, Warning, TEXT("Could not set sequence on sequence player, value is not dynamic. Set it as Always Dynamic."));
+			}
+
+			if(bAnimSequenceChanged && BlendTime > 0.0f)
+			{
+				if (const FAnimationUpdateContext* AnimationUpdateContext = UpdateContext.GetContext())
+				{
+					if (UE::Anim::IInertializationRequester* InertializationRequester = AnimationUpdateContext->GetMessage<UE::Anim::IInertializationRequester>())
+					{
+						InertializationRequester->RequestInertialization(BlendTime);
+					}
+				}
+				else
+				{
+					UE_LOG(LogSequencePlayerLibrary, Warning, TEXT("SetSequenceWithInertialBlending called with invalid context"));
+				}
+			}
 		});
 
-	return SequencePlayerContext;
+	return SequencePlayer;
 }
\ No newline at end of file
diff --git a/Engine/Source/Runtime/AnimGraphRuntime/Public/AnimNodes/AnimNode_SequenceEvaluator.h b/Engine/Source/Runtime/AnimGraphRuntime/Public/AnimNodes/AnimNode_SequenceEvaluator.h
--- a/Engine/Source/Runtime/AnimGraphRuntime/Public/AnimNodes/AnimNode_SequenceEvaluator.h
+++ b/Engine/Source/Runtime/AnimGraphRuntime/Public/AnimNodes/AnimNode_SequenceEvaluator.h
@@ -62,7 +62,7 @@ public:
 	virtual bool SetSequence(UAnimSequenceBase* InSequence) { return false; }
 
 	// Set the time at which to evaluate the associated sequence
-	virtual void SetExplicitTime(float InTime) {}
+	virtual bool SetExplicitTime(float InTime) { return false; }
 
 	// Set whether to teleport to explicit time when it is set
 	virtual void SetTeleportToExplicitTime(bool bInTeleport) {}
@@ -149,6 +149,7 @@ public:
 	virtual bool SetSequence(UAnimSequenceBase* InSequence) override;
 	virtual UAnimSequenceBase* GetSequence() const override;
 	virtual float GetExplicitTime() const override;
+	virtual bool SetExplicitTime(float InTime) override;
 	virtual bool GetShouldLoop() const override;
 	virtual bool GetTeleportToExplicitTime() const override;
 	virtual TEnumAsByte<ESequenceEvalReinit::Type> GetReinitializationBehavior() const override;
@@ -217,7 +218,7 @@ private:
 public:
 	// FAnimNode_SequenceEvaluatorBase interface
 	virtual bool SetSequence(UAnimSequenceBase* InSequence) override { Sequence = InSequence; return true; }
-	virtual void SetExplicitTime(float InTime) override { ExplicitTime = InTime; }
+	virtual bool SetExplicitTime(float InTime) override { ExplicitTime = InTime; return true; }
 	virtual void SetTeleportToExplicitTime(bool bInTeleport) override { bTeleportToExplicitTime = bInTeleport; }
 	virtual void SetReinitializationBehavior(TEnumAsByte<ESequenceEvalReinit::Type> InBehavior) override { ReinitializationBehavior = InBehavior; }
 	virtual UAnimSequenceBase* GetSequence() const override { return Sequence; }
diff --git a/Engine/Source/Runtime/AnimGraphRuntime/Public/SequenceEvaluatorLibrary.h b/Engine/Source/Runtime/AnimGraphRuntime/Public/SequenceEvaluatorLibrary.h
new file mode 100644
--- /dev/null
+++ b/Engine/Source/Runtime/AnimGraphRuntime/Public/SequenceEvaluatorLibrary.h
@@ -0,0 +1,46 @@
+// Copyright Epic Games, Inc. All Rights Reserved.
+
+#pragma once
+
+#include "CoreMinimal.h"
+
+#include "Animation/AnimExecutionContext.h"
+#include "Kismet/BlueprintFunctionLibrary.h"
+#include "Animation/AnimNodeReference.h"
+#include "SequenceEvaluatorLibrary.generated.h"
+
+struct FAnimNode_SequenceEvaluator;
+
+USTRUCT(BlueprintType)
+struct FSequenceEvaluatorReference : public FAnimNodeReference
+{
+	GENERATED_BODY()
+
+	typedef FAnimNode_SequenceEvaluator FInternalNodeType;
+};
+
+// Exposes operations to be performed on a sequence evaluator anim node
+// Note: Experimental and subject to change!
+UCLASS(Experimental)
+class ANIMGRAPHRUNTIME_API USequenceEvaluatorLibrary : public UBlueprintFunctionLibrary
+{
+	GENERATED_BODY()
+
+public:
+	/** Get a sequence evaluator context from an anim node context */
+	UFUNCTION(BlueprintCallable, Category = "Sequence Evaluator", meta=(BlueprintThreadSafe, ExpandEnumAsExecs = "Result"))
+	static FSequenceEvaluatorReference ConvertToSequenceEvaluatorContext(const FAnimNodeReference& Node, EAnimNodeReferenceConversionResult& Result);
+	
+	/** Set the current accumulated time of the sequence evaluator */
+	UFUNCTION(BlueprintCallable, Category = "Sequence Evaluator", meta=(BlueprintThreadSafe))
+	static FSequenceEvaluatorReference SetExplicitTime(const FSequenceEvaluatorReference& SequenceEvaluator, float Time);
+
+	/** Set the current sequence of the sequence evaluator */
+	UFUNCTION(BlueprintCallable, Category = "Sequence Evaluator", meta=(BlueprintThreadSafe))
+	static FSequenceEvaluatorReference SetSequence(const FSequenceEvaluatorReference& SequenceEvaluator, UAnimSequenceBase* Sequence);
+
+	/** Set the current sequence of the sequence evaluator with an inertial blend time */
+	UFUNCTION(BlueprintCallable, Category = "Sequence Evaluator", meta=(BlueprintThreadSafe))
+	static FSequenceEvaluatorReference SetSequenceWithInertialBlending(const FAnimUpdateContext& UpdateContext, const FSequenceEvaluatorReference& SequenceEvaluator, UAnimSequenceBase* Sequence, float BlendTime = 0.2f);
+	
+};
\ No newline at end of file
diff --git a/Engine/Source/Runtime/AnimGraphRuntime/Public/SequencePlayerLibrary.h b/Engine/Source/Runtime/AnimGraphRuntime/Public/SequencePlayerLibrary.h
--- a/Engine/Source/Runtime/AnimGraphRuntime/Public/SequencePlayerLibrary.h
+++ b/Engine/Source/Runtime/AnimGraphRuntime/Public/SequencePlayerLibrary.h
@@ -3,6 +3,8 @@
 #pragma once
 
 #include "CoreMinimal.h"
+
+#include "Animation/AnimExecutionContext.h"
 #include "Kismet/BlueprintFunctionLibrary.h"
 #include "Animation/AnimNodeReference.h"
 #include "SequencePlayerLibrary.generated.h"
@@ -37,6 +39,16 @@ public:
 	UFUNCTION(BlueprintCallable, Category = "Sequence Player", meta=(BlueprintThreadSafe))
 	static FSequencePlayerReference SetStartPosition(const FSequencePlayerReference& SequencePlayer, float StartPosition);
 
+	/** Set the play rate of the sequence player */
 	UFUNCTION(BlueprintCallable, Category = "Sequence Player", meta=(BlueprintThreadSafe))
 	static FSequencePlayerReference SetPlayRate(const FSequencePlayerReference& SequencePlayer, float PlayRate);
+
+	/** Set the current sequence of the sequence player */
+	UFUNCTION(BlueprintCallable, Category = "Sequence Player", meta=(BlueprintThreadSafe))
+	static FSequencePlayerReference SetSequence(const FSequencePlayerReference& SequencePlayer, UAnimSequenceBase* Sequence);
+
+	/** Set the current sequence of the sequence player with an inertial blend time */
+	UFUNCTION(BlueprintCallable, Category = "Sequence Player", meta=(BlueprintThreadSafe))
+	static FSequencePlayerReference SetSequenceWithInertialBlending(const FAnimUpdateContext& UpdateContext, const FSequencePlayerReference& SequencePlayer, UAnimSequenceBase* Sequence, float BlendTime = 0.2f);
+	
 };
\ No newline at end of file
diff --git a/Engine/Source/Runtime/Engine/Classes/Animation/AnimExecutionContext.h b/Engine/Source/Runtime/Engine/Classes/Animation/AnimExecutionContext.h
--- a/Engine/Source/Runtime/Engine/Classes/Animation/AnimExecutionContext.h
+++ b/Engine/Source/Runtime/Engine/Classes/Animation/AnimExecutionContext.h
@@ -79,7 +79,7 @@ public:
 	template<typename OtherContextType>
 	static OtherContextType ConvertToType(const FAnimExecutionContext& InContext, EAnimExecutionContextConversionResult& OutResult)
 	{
-		static_assert(TIsDerivedFrom<OtherContextType, FAnimExecutionContext>::IsDerived, "Argument ContextType must derive from FAnimExecutionContext");
+		static_assert(TIsDerivedFrom<OtherContextType, FAnimExecutionContext>::IsDerived, "Argument OtherContextType must derive from FAnimExecutionContext");
 		
 		if(TSharedPtr<FData> PinnedData = InContext.Data.Pin())
 		{
@@ -98,27 +98,28 @@ public:
 		return OtherContextType();
 	}
 
-	// Access internal context. Will return nullptr if invalid or an incorrect type is requested
-	template<typename InternalContextType>
-	InternalContextType* GetContext() const
+	// Access internal context. Will return nullptr if invalid
+	FAnimationBaseContext* GetBaseContext() const
 	{
 		if(TSharedPtr<FData> PinnedData = Data.Pin())
 		{
-			if(InternalContextType::InternalContextTypeId == PinnedData->ContextType)
-			{
-				return static_cast<InternalContextType*>(PinnedData->Context);
-			}
+			return PinnedData->Context;
 		}
 
 		return nullptr;
 	}
 
-	// Access internal context. Will return nullptr if invalid
-	FAnimationBaseContext* GetBaseContext() const
+protected:
+	// Access internal context. Will return nullptr if invalid or an incorrect type is requested
+	template<typename OtherContextType, typename InternalContextType>
+	typename InternalContextType* GetInternalContext() const
 	{
 		if(TSharedPtr<FData> PinnedData = Data.Pin())
 		{
-			return PinnedData->Context;
+			if(OtherContextType::InternalContextTypeId == PinnedData->ContextType)
+			{
+				return static_cast<InternalContextType*>(PinnedData->Context);
+			}
 		}
 
 		return nullptr;
@@ -130,12 +131,12 @@ protected:
 };
 
 USTRUCT()
-struct FAnimInitializationContext : public FAnimExecutionContext
+struct ENGINE_API FAnimInitializationContext : public FAnimExecutionContext
 {
 	GENERATED_BODY()
 
 public:
-	static const FAnimExecutionContext::FData::EContextType InternalContextTypeId = FData::EContextType::Initialize;
+	static const FData::EContextType InternalContextTypeId = FData::EContextType::Initialize;
 
 	FAnimInitializationContext() = default;
 	
@@ -145,10 +146,12 @@ public:
 		check(InData.IsValid());
 		check(InData.Pin()->ContextType == InternalContextTypeId);
 	}
+
+	FAnimationInitializeContext* GetContext() const;
 };
 
 USTRUCT()
-struct FAnimUpdateContext : public FAnimExecutionContext
+struct ENGINE_API FAnimUpdateContext : public FAnimExecutionContext
 {
 	GENERATED_BODY()
 
@@ -163,10 +166,12 @@ public:
 		check(InData.IsValid());
 		check(InData.Pin()->ContextType == InternalContextTypeId);
 	}
+
+	FAnimationUpdateContext* GetContext() const;
 };
 
 USTRUCT()
-struct FAnimPoseContext : public FAnimExecutionContext
+struct ENGINE_API FAnimPoseContext : public FAnimExecutionContext
 {
 	GENERATED_BODY()
 
@@ -181,10 +186,12 @@ public:
 		check(InData.IsValid());
 		check(InData.Pin()->ContextType == InternalContextTypeId);
 	}
+
+	FPoseContext* GetContext() const;
 };
 
 USTRUCT()
-struct FAnimComponentSpacePoseContext : public FAnimExecutionContext
+struct ENGINE_API FAnimComponentSpacePoseContext : public FAnimExecutionContext
 {
 	GENERATED_BODY()
 
@@ -199,4 +206,6 @@ public:
 		check(InData.IsValid());
 		check(InData.Pin()->ContextType == InternalContextTypeId);
 	}
+
+	FComponentSpacePoseContext* GetContext() const;
 };
\ No newline at end of file
diff --git a/Engine/Source/Runtime/Engine/Private/Animation/AnimExecutionContext.cpp b/Engine/Source/Runtime/Engine/Private/Animation/AnimExecutionContext.cpp
--- a/Engine/Source/Runtime/Engine/Private/Animation/AnimExecutionContext.cpp
+++ b/Engine/Source/Runtime/Engine/Private/Animation/AnimExecutionContext.cpp
@@ -26,3 +26,23 @@ FAnimExecutionContext::FData::FData(FComponentSpacePoseContext& InContext)
 	Context = &InContext;
 	ContextType = EContextType::ComponentSpacePose;
 }
+
+FAnimationInitializeContext* FAnimInitializationContext::GetContext() const
+{
+	return GetInternalContext<FAnimInitializationContext, FAnimationInitializeContext>();
+}
+
+FAnimationUpdateContext* FAnimUpdateContext::GetContext() const
+{
+	return GetInternalContext<FAnimUpdateContext, FAnimationUpdateContext>();
+}
+
+FPoseContext* FAnimPoseContext::GetContext() const
+{
+	return GetInternalContext<FAnimPoseContext, FPoseContext>();
+}
+
+FComponentSpacePoseContext* FAnimComponentSpacePoseContext::GetContext() const
+{
+	return GetInternalContext<FAnimPoseContext, FComponentSpacePoseContext>();
+}
\ No newline at end of file
