Missed file

diff --git a/Engine/Source/Runtime/AnimGraphRuntime/Public/AnimExecutionContextLibrary.h b/Engine/Source/Runtime/AnimGraphRuntime/Public/AnimExecutionContextLibrary.h
--- a/Engine/Source/Runtime/AnimGraphRuntime/Public/AnimExecutionContextLibrary.h
+++ b/Engine/Source/Runtime/AnimGraphRuntime/Public/AnimExecutionContextLibrary.h
@@ -28,6 +28,22 @@ public:
 	static UAnimInstance* GetAnimInstance(const FAnimExecutionContext& Context);
 
 	/** Internal compiler use only - Get a reference to an anim node by index */
-	UFUNCTION(BlueprintPure, BlueprintInternalUseOnly, meta=(BlueprintThreadSafe))
+	UFUNCTION(BlueprintPure, BlueprintInternalUseOnly, meta=(BlueprintThreadSafe, DefaultToSelf = "Instance"))
 	static FAnimNodeReference GetAnimNodeReference(UAnimInstance* Instance, int32 Index);
+
+	/** Convert to an initialization context */
+	UFUNCTION(BlueprintCallable, Category = "Anim Node", meta=(BlueprintThreadSafe, ExpandEnumAsExecs = "Result"))
+	FAnimInitializationContext ConvertToInitializationContext(const FAnimExecutionContext& Context, EAnimExecutionContextConversionResult& Result);
+	
+	/** Convert to an update context */
+	UFUNCTION(BlueprintCallable, Category = "Anim Node", meta=(BlueprintThreadSafe, ExpandEnumAsExecs = "Result"))
+	static FAnimUpdateContext ConvertToUpdateContext(const FAnimExecutionContext& Context, EAnimExecutionContextConversionResult& Result);
+
+	/** Convert to a pose context */
+	UFUNCTION(BlueprintCallable, Category = "Anim Node", meta=(BlueprintThreadSafe, ExpandEnumAsExecs = "Result"))
+	static FAnimPoseContext ConvertToPoseContext(const FAnimExecutionContext& Context, EAnimExecutionContextConversionResult& Result);
+
+	/** Convert to a component space pose context */
+	UFUNCTION(BlueprintCallable, Category = "Anim Node", meta=(BlueprintThreadSafe, ExpandEnumAsExecs = "Result"))
+	static FAnimComponentSpacePoseContext ConvertToComponentSpacePoseContext(const FAnimExecutionContext& Context, EAnimExecutionContextConversionResult& Result);
 };
\ No newline at end of file
