PR #8160: Fix error case NaN while delta time is 0 (Contributed by [REDACTED_PERSON])

Fixes a case in RBAN where mesh component initialisation may pass through a zero delta time which can be used to divide in the calculation of a velocity when the node is set to inherit bone velocities.

diff --git a/Engine/Source/Runtime/AnimGraphRuntime/Private/BoneControllers/AnimNode_RigidBody.cpp b/Engine/Source/Runtime/AnimGraphRuntime/Private/BoneControllers/AnimNode_RigidBody.cpp
--- a/Engine/Source/Runtime/AnimGraphRuntime/Private/BoneControllers/AnimNode_RigidBody.cpp
+++ b/Engine/Source/Runtime/AnimGraphRuntime/Private/BoneControllers/AnimNode_RigidBody.cpp
@@ -688,20 +688,22 @@ void FAnimNode_RigidBody::EvaluateSkeletalControl_AnyThread(FComponentSpacePoseC
 							const FTransform PrevSSTM = ConvertCSTransformToSimSpace(SimulationSpace, PrevCSTM, CompWorldSpaceTM, BaseBoneTM);
 							const FTransform NextSSTM = ConvertCSTransformToSimSpace(SimulationSpace, NextCSTM, CompWorldSpaceTM, BaseBoneTM);
 
-							// Linear Velocity
 							if(DeltaSeconds > 0.0f)
 							{
+								// Linear Velocity
 								BodyData.TransferedBoneLinearVelocity = ((NextSSTM.GetLocation() - PrevSSTM.GetLocation()) / DeltaSeconds);
+								
+								// Angular Velocity
+								const FQuat DeltaRotation = (NextSSTM.GetRotation().Inverse() * PrevSSTM.GetRotation());
+								const float RotationAngle = DeltaRotation.GetAngle() / DeltaSeconds;
+								BodyData.TransferedBoneAngularVelocity = (FQuat(DeltaRotation.GetRotationAxis(), RotationAngle)); 
 							}
 							else
 							{
 								BodyData.TransferedBoneLinearVelocity = (FVector::ZeroVector);
+								BodyData.TransferedBoneAngularVelocity = (FQuat::Identity); 
 							}
 
-							// Angular Velocity
-							const FQuat DeltaRotation = (NextSSTM.GetRotation().Inverse() * PrevSSTM.GetRotation());
-							const float RotationAngle = DeltaRotation.GetAngle() / DeltaSeconds;
-							BodyData.TransferedBoneAngularVelocity = (FQuat(DeltaRotation.GetRotationAxis(), RotationAngle)); 
 						}
 					}
 				}
