Added pure conversion functions for anim node references

diff --git a/Engine/Source/Runtime/AnimGraphRuntime/Private/LinkedAnimGraphLibrary.cpp b/Engine/Source/Runtime/AnimGraphRuntime/Private/LinkedAnimGraphLibrary.cpp
--- a/Engine/Source/Runtime/AnimGraphRuntime/Private/LinkedAnimGraphLibrary.cpp
+++ b/Engine/Source/Runtime/AnimGraphRuntime/Private/LinkedAnimGraphLibrary.cpp
@@ -4,7 +4,7 @@
 #include "Animation/AnimInstance.h"
 #include "Animation/AnimNode_LinkedAnimGraph.h"
 
-FLinkedAnimGraphReference ULinkedAnimGraphLibrary::ConvertToLinkedAnimGraphContext(const FAnimNodeReference& Node, EAnimNodeReferenceConversionResult& Result)
+FLinkedAnimGraphReference ULinkedAnimGraphLibrary::ConvertToLinkedAnimGraph(const FAnimNodeReference& Node, EAnimNodeReferenceConversionResult& Result)
 {
 	return FAnimNodeReference::ConvertToType<FLinkedAnimGraphReference>(Node, Result);
 }
diff --git a/Engine/Source/Runtime/AnimGraphRuntime/Private/SequenceEvaluatorLibrary.cpp b/Engine/Source/Runtime/AnimGraphRuntime/Private/SequenceEvaluatorLibrary.cpp
--- a/Engine/Source/Runtime/AnimGraphRuntime/Private/SequenceEvaluatorLibrary.cpp
+++ b/Engine/Source/Runtime/AnimGraphRuntime/Private/SequenceEvaluatorLibrary.cpp
@@ -7,7 +7,7 @@
 
 DEFINE_LOG_CATEGORY_STATIC(LogSequenceEvaluatorLibrary, Verbose, All);
 
-FSequenceEvaluatorReference USequenceEvaluatorLibrary::ConvertToSequenceEvaluatorContext(const FAnimNodeReference& Node, EAnimNodeReferenceConversionResult& Result)
+FSequenceEvaluatorReference USequenceEvaluatorLibrary::ConvertToSequenceEvaluator(const FAnimNodeReference& Node, EAnimNodeReferenceConversionResult& Result)
 {
 	return FAnimNodeReference::ConvertToType<FSequenceEvaluatorReference>(Node, Result);
 }
diff --git a/Engine/Source/Runtime/AnimGraphRuntime/Private/SequencePlayerLibrary.cpp b/Engine/Source/Runtime/AnimGraphRuntime/Private/SequencePlayerLibrary.cpp
--- a/Engine/Source/Runtime/AnimGraphRuntime/Private/SequencePlayerLibrary.cpp
+++ b/Engine/Source/Runtime/AnimGraphRuntime/Private/SequencePlayerLibrary.cpp
@@ -7,7 +7,7 @@
 
 DEFINE_LOG_CATEGORY_STATIC(LogSequencePlayerLibrary, Verbose, All);
 
-FSequencePlayerReference USequencePlayerLibrary::ConvertToSequencePlayerContext(const FAnimNodeReference& Node, EAnimNodeReferenceConversionResult& Result)
+FSequencePlayerReference USequencePlayerLibrary::ConvertToSequencePlayer(const FAnimNodeReference& Node, EAnimNodeReferenceConversionResult& Result)
 {
 	return FAnimNodeReference::ConvertToType<FSequencePlayerReference>(Node, Result);
 }
diff --git a/Engine/Source/Runtime/AnimGraphRuntime/Public/LinkedAnimGraphLibrary.h b/Engine/Source/Runtime/AnimGraphRuntime/Public/LinkedAnimGraphLibrary.h
--- a/Engine/Source/Runtime/AnimGraphRuntime/Public/LinkedAnimGraphLibrary.h
+++ b/Engine/Source/Runtime/AnimGraphRuntime/Public/LinkedAnimGraphLibrary.h
@@ -24,9 +24,18 @@ class ANIMGRAPHRUNTIME_API ULinkedAnimGraphLibrary : public UBlueprintFunctionLi
 	GENERATED_BODY()
 
 public:
-	/** Get a sequence player context from an anim node context */
+	/** Get a linked anim graph reference from an anim node reference */
 	UFUNCTION(BlueprintCallable, Category = "Linked Anim Graph", meta=(BlueprintThreadSafe, ExpandEnumAsExecs = "Result"))
-	static FLinkedAnimGraphReference ConvertToLinkedAnimGraphContext(const FAnimNodeReference& Node, EAnimNodeReferenceConversionResult& Result);
+	static FLinkedAnimGraphReference ConvertToLinkedAnimGraph(const FAnimNodeReference& Node, EAnimNodeReferenceConversionResult& Result);
+
+	/** Get a linked anim graph reference from an anim node reference (pure) */
+	UFUNCTION(BlueprintPure, Category = "Linked Anim Graph", meta=(BlueprintThreadSafe, DisplayName = "Convert to Linked Anim Graph"))
+	static void ConvertToLinkedAnimGraphPure(const FAnimNodeReference& Node, FLinkedAnimGraphReference& LinkedAnimGraph, bool& Result)
+	{
+		EAnimNodeReferenceConversionResult ConversionResult;
+		LinkedAnimGraph = ConvertToLinkedAnimGraph(Node, ConversionResult);
+		Result = (ConversionResult == EAnimNodeReferenceConversionResult::Succeeded);
+	}
 	
 	/** Returns whether the node hosts an instance (e.g. linked anim graph or layer) */
 	UFUNCTION(BlueprintPure, Category = "Linked Anim Graph", meta=(BlueprintThreadSafe))
diff --git a/Engine/Source/Runtime/AnimGraphRuntime/Public/SequenceEvaluatorLibrary.h b/Engine/Source/Runtime/AnimGraphRuntime/Public/SequenceEvaluatorLibrary.h
--- a/Engine/Source/Runtime/AnimGraphRuntime/Public/SequenceEvaluatorLibrary.h
+++ b/Engine/Source/Runtime/AnimGraphRuntime/Public/SequenceEvaluatorLibrary.h
@@ -29,7 +29,16 @@ class ANIMGRAPHRUNTIME_API USequenceEvaluatorLibrary : public UBlueprintFunction
 public:
 	/** Get a sequence evaluator context from an anim node context */
 	UFUNCTION(BlueprintCallable, Category = "Sequence Evaluator", meta=(BlueprintThreadSafe, ExpandEnumAsExecs = "Result"))
-	static FSequenceEvaluatorReference ConvertToSequenceEvaluatorContext(const FAnimNodeReference& Node, EAnimNodeReferenceConversionResult& Result);
+	static FSequenceEvaluatorReference ConvertToSequenceEvaluator(const FAnimNodeReference& Node, EAnimNodeReferenceConversionResult& Result);
+
+	/** Get a sequence evaluator context from an anim node context (pure) */
+	UFUNCTION(BlueprintPure, Category = "Sequence Evaluator", meta=(BlueprintThreadSafe, DisplayName = "Convert to Sequence Evaluator"))
+	static void ConvertToSequenceEvaluatorPure(const FAnimNodeReference& Node, FSequenceEvaluatorReference& SequenceEvaluator, bool& Result)
+	{
+		EAnimNodeReferenceConversionResult ConversionResult;
+		SequenceEvaluator = ConvertToSequenceEvaluator(Node, ConversionResult);
+		Result = (ConversionResult == EAnimNodeReferenceConversionResult::Succeeded);	
+	}
 	
 	/** Set the current accumulated time of the sequence evaluator */
 	UFUNCTION(BlueprintCallable, Category = "Sequence Evaluator", meta=(BlueprintThreadSafe))
diff --git a/Engine/Source/Runtime/AnimGraphRuntime/Public/SequencePlayerLibrary.h b/Engine/Source/Runtime/AnimGraphRuntime/Public/SequencePlayerLibrary.h
--- a/Engine/Source/Runtime/AnimGraphRuntime/Public/SequencePlayerLibrary.h
+++ b/Engine/Source/Runtime/AnimGraphRuntime/Public/SequencePlayerLibrary.h
@@ -29,8 +29,17 @@ class ANIMGRAPHRUNTIME_API USequencePlayerLibrary : public UBlueprintFunctionLib
 public:
 	/** Get a sequence player context from an anim node context */
 	UFUNCTION(BlueprintCallable, Category = "Sequence Player", meta=(BlueprintThreadSafe, ExpandEnumAsExecs = "Result"))
-	static FSequencePlayerReference ConvertToSequencePlayerContext(const FAnimNodeReference& Node, EAnimNodeReferenceConversionResult& Result);
-	
+	static FSequencePlayerReference ConvertToSequencePlayer(const FAnimNodeReference& Node, EAnimNodeReferenceConversionResult& Result);
+
+	/** Get a sequence player context from an anim node context (pure) */
+	UFUNCTION(BlueprintPure, Category = "Sequence Player", meta=(BlueprintThreadSafe, DisplayName = "Convert to Sequence Player"))
+	static void ConvertToSequencePlayerPure(const FAnimNodeReference& Node, FSequencePlayerReference& SequencePlayer, bool& Result)
+	{
+		EAnimNodeReferenceConversionResult ConversionResult;
+		SequencePlayer = ConvertToSequencePlayer(Node, ConversionResult);
+		Result = (ConversionResult == EAnimNodeReferenceConversionResult::Succeeded);
+	}
+
 	/** Set the current accumulated time of the sequence player */
 	UFUNCTION(BlueprintCallable, Category = "Sequence Player", meta=(BlueprintThreadSafe))
 	static FSequencePlayerReference SetAccumulatedTime(const FSequencePlayerReference& SequencePlayer, float Time);
