Fix in the weight based blend profiles, making them compatible with UE4 again. Thanks to [REDACTED_PERSON] and [REDACTED_PERSON] for finding this issue and diving in more.

diff --git a/Engine/Source/Runtime/AnimGraphRuntime/Private/AnimNodes/AnimNode_BlendListBase.cpp b/Engine/Source/Runtime/AnimGraphRuntime/Private/AnimNodes/AnimNode_BlendListBase.cpp
--- a/Engine/Source/Runtime/AnimGraphRuntime/Private/AnimNodes/AnimNode_BlendListBase.cpp
+++ b/Engine/Source/Runtime/AnimGraphRuntime/Private/AnimNodes/AnimNode_BlendListBase.cpp
@@ -239,7 +239,8 @@ void FAnimNode_BlendListBase::Update_AnyThread(const FAnimationUpdateContext& Co
 			{
 				FBlendSampleData& PoseSampleData = PerBoneSampleData[i];
 				PoseSampleData.TotalWeight = BlendWeights[i];
-				CurrentBlendProfile->UpdateBoneWeights(PoseSampleData, Blends[i], BlendStartAlphas[i], BlendWeights[i]);
+				const bool bInverse = (CurrentBlendProfile->Mode == EBlendProfileMode::WeightFactor) ? (ChildIndex != i) : false;
+				CurrentBlendProfile->UpdateBoneWeights(PoseSampleData, Blends[i], BlendStartAlphas[i], BlendWeights[i], bInverse);
 			}
 
 			FBlendSampleData::NormalizeDataWeight(PerBoneSampleData);
diff --git a/Engine/Source/Runtime/Engine/Private/Animation/BlendProfile.cpp b/Engine/Source/Runtime/Engine/Private/Animation/BlendProfile.cpp
--- a/Engine/Source/Runtime/Engine/Private/Animation/BlendProfile.cpp
+++ b/Engine/Source/Runtime/Engine/Private/Animation/BlendProfile.cpp
@@ -241,12 +241,12 @@ float UBlendProfile::CalculateBoneWeight(float BoneFactor, EBlendProfileMode Mod
 		{
 			if (!bInverse)
 			{
-				return FMath::Clamp(MainWeight * BoneFactor, ZERO_ANIMWEIGHT_THRESH, 1.0f);
+				return FMath::Max(MainWeight * BoneFactor, ZERO_ANIMWEIGHT_THRESH);
 			}
 
 			// We're inversing.
 			const float Weight = (BoneFactor > ZERO_ANIMWEIGHT_THRESH) ? MainWeight / BoneFactor : 1.0f;
-			return FMath::Clamp(Weight, ZERO_ANIMWEIGHT_THRESH, 1.0f);
+			return FMath::Max(Weight, ZERO_ANIMWEIGHT_THRESH);
 		}
 
 		// Handle unsupported modes.
