Added EvalFromRefPose

diff --git a/Engine/Source/Runtime/AnimGraphRuntime/Private/AnimNodes/AnimNode_PoseDriver.cpp b/Engine/Source/Runtime/AnimGraphRuntime/Private/AnimNodes/AnimNode_PoseDriver.cpp
--- a/Engine/Source/Runtime/AnimGraphRuntime/Private/AnimNodes/AnimNode_PoseDriver.cpp
+++ b/Engine/Source/Runtime/AnimGraphRuntime/Private/AnimNodes/AnimNode_PoseDriver.cpp
@@ -78,11 +78,15 @@ void FAnimNode_PoseDriver::CacheBones_AnyThread(const FAnimationCacheBonesContex
 {
 	DECLARE_SCOPE_HIERARCHICAL_COUNTER_ANIMNODE(CacheBones_AnyThread)
 	FAnimNode_PoseHandler::CacheBones_AnyThread(Context);
+
 	// Init pose input
 	SourcePose.CacheBones(Context);
 
 	const FBoneContainer& BoneContainer = Context.AnimInstanceProxy->GetRequiredBones();
 
+	// Get RefPose
+	RefPoseCompact = BoneContainer.GetRefPoseCompactArray();
+
 	// Init bone refs
 	for (FBoneReference& SourceBoneRef : SourceBones)
 	{
@@ -196,13 +200,45 @@ void FAnimNode_PoseDriver::GetRBFTargets(TArray<FRBFTarget>& OutTargets) const
 			if (PoseTarget.BoneTransforms.IsValidIndex(SourceIdx))
 			{
 				const FPoseDriverTransform& BoneTransform = PoseTarget.BoneTransforms[SourceIdx];
+
+				// Get Ref Transform
+				FTransform RefBoneTransform = FTransform::Identity;
+				if (bEvalFromRefPose)
+				{
+					const int32 CompactPoseIndex = SourceBones[SourceIdx].CachedCompactPoseIndex.GetInt();
+					if (RefPoseCompact.IsValidIndex(CompactPoseIndex))
+					{
+						RefBoneTransform = RefPoseCompact[CompactPoseIndex];
+					}
+				}
+
+				// Target Translation
 				if (DriveSource == EPoseDriverSource::Translation)
 				{
-					RBFTarget.AddFromVector(BoneTransform.TargetTranslation);
+					// Make translation relative to its Ref
+					if (bEvalFromRefPose)
+					{
+						RBFTarget.AddFromVector(RefBoneTransform.Inverse().TransformPosition(BoneTransform.TargetTranslation));
+					}
+					else
+					{
+						RBFTarget.AddFromVector(BoneTransform.TargetTranslation);
+					}
 				}
+
+				// Target Rotation
 				else
 				{
-					RBFTarget.AddFromRotator(BoneTransform.TargetRotation);
+					// Make rotation relative to its Ref
+					if (bEvalFromRefPose)
+					{	
+						const FQuat TargetRotation = BoneTransform.TargetRotation.Quaternion();
+						RBFTarget.AddFromRotator(RefBoneTransform.Inverse().TransformRotation(TargetRotation).Rotator());
+					}
+					else
+					{
+						RBFTarget.AddFromRotator(BoneTransform.TargetRotation);
+					}
 				}
 			}
 			else
@@ -272,7 +308,15 @@ void FAnimNode_PoseDriver::Evaluate_AnyThread(FPoseContext& Output)
 				// If just evaluating in local space, just grab from local space pose
 				else
 				{
-					SourceBoneTM = SourceData.Pose[SourceCompactIndex];
+					// Relative to Ref Pose
+					if (bEvalFromRefPose && RefPoseCompact.IsValidIndex(SourceCompactIndex.GetInt()))
+					{
+						SourceBoneTM = SourceData.Pose[SourceCompactIndex].GetRelativeTransform(RefPoseCompact[SourceCompactIndex.GetInt()]);
+					}
+					else
+					{
+						SourceBoneTM = SourceData.Pose[SourceCompactIndex];
+					}
 				}
 
 				bFoundAnyBone = true;
diff --git a/Engine/Source/Runtime/AnimGraphRuntime/Public/AnimNodes/AnimNode_PoseDriver.h b/Engine/Source/Runtime/AnimGraphRuntime/Public/AnimNodes/AnimNode_PoseDriver.h
--- a/Engine/Source/Runtime/AnimGraphRuntime/Public/AnimNodes/AnimNode_PoseDriver.h
+++ b/Engine/Source/Runtime/AnimGraphRuntime/Public/AnimNodes/AnimNode_PoseDriver.h
@@ -164,6 +164,16 @@ struct ANIMGRAPHRUNTIME_API FAnimNode_PoseDriver : public FAnimNode_PoseHandler
 	UPROPERTY(EditAnywhere, Category = PoseDriver)
 	FBoneReference EvalSpaceBone;
 
+	/**
+	 *	Evaluate SourceBone transform relative from its Reference Pose.
+	 *  This is recommended when using Swing and Twist Angle as Distance Method, since the twist will be computed from RefPose.
+	 * 
+	 *	If not specified, we just use local space of SourceBone (ie relative to parent bone)
+	 *  This mode won't work in conjunction with EvalSpaceBone;
+	 */
+	UPROPERTY(EditAnywhere, Category = PoseDriver)
+	bool bEvalFromRefPose = false;
+
 	/** Parameters used by RBF solver */
 	UPROPERTY(EditAnywhere, Category = PoseDriver, meta=(ShowOnlyInnerProperties))
 	FRBFParams RBFParams;
@@ -242,4 +252,5 @@ private:
 	TSharedPtr<const FRBFSolverData> SolverData;
 	FRBFEntry RBFInput;
 	TArray<FRBFTarget> RBFTargets;
+	TArray<FTransform> RefPoseCompact;
 };
