Convert FParallelAnimationEvaluationTask to use taskgraph completion event for rigid body animation node.
Fixes anim task blocking on RB sim task and doing 0 work, now taskgraph schedules work and doesn't lock up extra threads.

diff --git a/Engine/Source/Runtime/AnimGraphRuntime/Private/BoneControllers/AnimNode_RigidBody.cpp b/Engine/Source/Runtime/AnimGraphRuntime/Private/BoneControllers/AnimNode_RigidBody.cpp
--- a/Engine/Source/Runtime/AnimGraphRuntime/Private/BoneControllers/AnimNode_RigidBody.cpp
+++ b/Engine/Source/Runtime/AnimGraphRuntime/Private/BoneControllers/AnimNode_RigidBody.cpp
@@ -528,14 +528,6 @@ public:
 	void DoTask(ENamedThreads::Type CurrentThread, const FGraphEventRef& CompletionGraphEvent)
 	{
 		RigidBodyNode->RunPhysicsSimulation(DeltaSeconds, SimSpaceGravity);
-
-		// Signal that the node does not need to wait on the event, but trigger it if it is already in the process of doing so
-		FAnimNode_RigidBody::FSimulationTaskState& TaskState = RigidBodyNode->SimulationTaskState;
-		const bool bWasSimulationPending = TaskState.bSimulationPending.exchange(false);
-		if (!bWasSimulationPending)
-		{
-			TaskState.SimulationCompletionEvent->Trigger();
-		}
 	}
 
 private:
@@ -568,14 +560,12 @@ void FAnimNode_RigidBody::RunPhysicsSimulation(float DeltaSeconds, const FVector
 
 void FAnimNode_RigidBody::FlushDeferredSimulationTask()
 {
-	// If the task hasn't lowered the pending flag by this point, we must wait for the event
-	const bool bWasSimulationPending = SimulationTaskState.bSimulationPending.exchange(false);
-	if (bWasSimulationPending)
+	FGraphEventRef& CompletionEvent = SimulationTaskState.SimulationCompletionEvent;
+	if (CompletionEvent && !CompletionEvent->IsComplete())
 	{
-		check(SimulationTaskState.SimulationCompletionEvent);
 		SCOPE_CYCLE_COUNTER(STAT_RigidBodyNode_SimulationWait);
 		CSV_SCOPED_TIMING_STAT(Animation, RigidBodyNodeSimulationWait);
-		SimulationTaskState.SimulationCompletionEvent->Wait();
+		FTaskGraphInterface::Get().WaitUntilTasksComplete({CompletionEvent});
 	}
 }
 
@@ -584,11 +574,6 @@ void FAnimNode_RigidBody::DestroyPhysicsSimulation()
 	FlushDeferredSimulationTask();
 	delete PhysicsSimulation;
 	PhysicsSimulation = nullptr;
-	if (SimulationTaskState.SimulationCompletionEvent)
-	{
-		FPlatformProcess::ReturnSynchEventToPool(SimulationTaskState.SimulationCompletionEvent);
-		SimulationTaskState.SimulationCompletionEvent = nullptr;
-	}
 }
 
 void FAnimNode_RigidBody::EvaluateSkeletalControl_AnyThread(FComponentSpacePoseContext& Output, TArray<FBoneTransform>& OutBoneTransforms)
@@ -987,13 +972,7 @@ void FAnimNode_RigidBody::EvaluateSkeletalControl_AnyThread(FComponentSpacePoseC
 		// Deferred task must be started after we read actor poses to avoid a race
 		if (bNeedsSimulationTick && bUseDeferredSimulationTask)
 		{
-			if (!SimulationTaskState.SimulationCompletionEvent)
-			{
-				const bool bIsManualReset = false;
-				SimulationTaskState.SimulationCompletionEvent = FPlatformProcess::GetSynchEventFromPool(bIsManualReset);
-			}
-			SimulationTaskState.bSimulationPending.store(true);
-			TGraphTask<FRigidBodyNodeSimulationTask>::CreateTask().ConstructAndDispatchWhenReady(this, DeltaSeconds, SimSpaceGravity);
+			SimulationTaskState.SimulationCompletionEvent = TGraphTask<FRigidBodyNodeSimulationTask>::CreateTask().ConstructAndDispatchWhenReady(this, DeltaSeconds, SimSpaceGravity);
 		}
 
 		PreviousCompWorldSpaceTM = CompWorldSpaceTM;
diff --git a/Engine/Source/Runtime/AnimGraphRuntime/Public/BoneControllers/AnimNode_RigidBody.h b/Engine/Source/Runtime/AnimGraphRuntime/Public/BoneControllers/AnimNode_RigidBody.h
--- a/Engine/Source/Runtime/AnimGraphRuntime/Public/BoneControllers/AnimNode_RigidBody.h
+++ b/Engine/Source/Runtime/AnimGraphRuntime/Public/BoneControllers/AnimNode_RigidBody.h
@@ -346,7 +346,6 @@ private:
 	{
 		FSimulationTaskState()
 			: SimulationCompletionEvent(nullptr)
-			, bSimulationPending(false)
 		{
 		}
 
@@ -357,8 +356,7 @@ private:
 			return(*this);
 		}
 
-		FEvent* SimulationCompletionEvent;
-		std::atomic<bool> bSimulationPending;
+		FGraphEventRef SimulationCompletionEvent;
 	};
 	FSimulationTaskState SimulationTaskState;
 
