Fix FRigidBodyNodeSimulationTask not using TrackSubsequents giving null a FGraphEventRef causing threading crashes.

diff --git a/Engine/Source/Runtime/AnimGraphRuntime/Private/BoneControllers/AnimNode_RigidBody.cpp b/Engine/Source/Runtime/AnimGraphRuntime/Private/BoneControllers/AnimNode_RigidBody.cpp
--- a/Engine/Source/Runtime/AnimGraphRuntime/Private/BoneControllers/AnimNode_RigidBody.cpp
+++ b/Engine/Source/Runtime/AnimGraphRuntime/Private/BoneControllers/AnimNode_RigidBody.cpp
@@ -522,7 +522,7 @@ public:
 
 	static FORCEINLINE ESubsequentsMode::Type GetSubsequentsMode()
 	{
-		return ESubsequentsMode::FireAndForget;
+		return ESubsequentsMode::TrackSubsequents;
 	}
 
 	void DoTask(ENamedThreads::Type CurrentThread, const FGraphEventRef& CompletionGraphEvent)
@@ -972,6 +972,8 @@ void FAnimNode_RigidBody::EvaluateSkeletalControl_AnyThread(FComponentSpacePoseC
 		// Deferred task must be started after we read actor poses to avoid a race
 		if (bNeedsSimulationTick && bUseDeferredSimulationTask)
 		{
+			// FlushDeferredSimulationTask() should have already ensured task is done.
+			ensure(!SimulationTaskState.SimulationCompletionEvent || SimulationTaskState.SimulationCompletionEvent->IsComplete());
 			SimulationTaskState.SimulationCompletionEvent = TGraphTask<FRigidBodyNodeSimulationTask>::CreateTask().ConstructAndDispatchWhenReady(this, DeltaSeconds, SimSpaceGravity);
 		}
 
