Skeletal control BP function library

Also fixes type conversions (first time I have tried to convert to a base class and found that my IsChildOf check was the wrong way around).
Plus some extra accessors added for exec/update contexts.
Found the need for these when doing show & tell prep.

diff --git a/Engine/Source/Runtime/AnimGraphRuntime/Private/AnimExecutionContextLibrary.cpp b/Engine/Source/Runtime/AnimGraphRuntime/Private/AnimExecutionContextLibrary.cpp
--- a/Engine/Source/Runtime/AnimGraphRuntime/Private/AnimExecutionContextLibrary.cpp
+++ b/Engine/Source/Runtime/AnimGraphRuntime/Private/AnimExecutionContextLibrary.cpp
@@ -5,6 +5,7 @@
 #include "Animation/AnimInstance.h"
 #include "Animation/AnimNodeBase.h"
 #include "Animation/AnimClassInterface.h"
+#include "Animation/AnimInstanceProxy.h"
 
 UAnimInstance* UAnimExecutionContextLibrary::GetAnimInstance(const FAnimExecutionContext& Context)
 {
@@ -31,6 +32,26 @@ FAnimUpdateContext UAnimExecutionContextLibrary::ConvertToUpdateContext(const FA
 	return FAnimExecutionContext::ConvertToType<FAnimUpdateContext>(Context, Result);
 }
 
+float UAnimExecutionContextLibrary::GetDeltaTime(const FAnimUpdateContext& UpdateContext)
+{
+	if(FAnimationUpdateContext* Context = UpdateContext.GetContext())
+	{
+		return Context->GetDeltaTime();
+	}
+
+	return 0.0f;
+}
+
+float UAnimExecutionContextLibrary::GetCurrentWeight(const FAnimUpdateContext& UpdateContext)
+{
+	if(FAnimationUpdateContext* Context = UpdateContext.GetContext())
+	{
+		return Context->GetFinalBlendWeight();
+	}
+
+	return 0.0f;
+}
+
 FAnimPoseContext UAnimExecutionContextLibrary::ConvertToPoseContext(const FAnimExecutionContext& Context, EAnimExecutionContextConversionResult& Result)
 {
 	return FAnimExecutionContext::ConvertToType<FAnimPoseContext>(Context, Result);
diff --git a/Engine/Source/Runtime/AnimGraphRuntime/Private/BoneControllers/AnimNode_SkeletalControlBase.cpp b/Engine/Source/Runtime/AnimGraphRuntime/Private/BoneControllers/AnimNode_SkeletalControlBase.cpp
--- a/Engine/Source/Runtime/AnimGraphRuntime/Private/BoneControllers/AnimNode_SkeletalControlBase.cpp
+++ b/Engine/Source/Runtime/AnimGraphRuntime/Private/BoneControllers/AnimNode_SkeletalControlBase.cpp
@@ -153,3 +153,13 @@ void FAnimNode_SkeletalControlBase::EvaluateSkeletalControl_AnyThread(FComponent
 	PRAGMA_ENABLE_DEPRECATION_WARNINGS
 }
 
+void FAnimNode_SkeletalControlBase::SetAlpha(float InAlpha)
+{
+	Alpha = InAlpha;
+	ActualAlpha = InAlpha;
+}
+
+float FAnimNode_SkeletalControlBase::GetAlpha() const
+{
+	return ActualAlpha;
+}
\ No newline at end of file
diff --git a/Engine/Source/Runtime/AnimGraphRuntime/Private/SkeletalControlLibrary.cpp b/Engine/Source/Runtime/AnimGraphRuntime/Private/SkeletalControlLibrary.cpp
new file mode 100644
--- /dev/null
+++ b/Engine/Source/Runtime/AnimGraphRuntime/Private/SkeletalControlLibrary.cpp
@@ -0,0 +1,38 @@
+// Copyright Epic Games, Inc. All Rights Reserved.
+
+#include "SkeletalControlLibrary.h"
+
+#include "BoneControllers/AnimNode_SkeletalControlBase.h"
+
+DEFINE_LOG_CATEGORY_STATIC(LogSkeletalControlLibrary, Verbose, All);
+
+FSkeletalControlReference USkeletalControlLibrary::ConvertToSkeletalControl(const FAnimNodeReference& Node, EAnimNodeReferenceConversionResult& Result)
+{
+	return FAnimNodeReference::ConvertToType<FSkeletalControlReference>(Node, Result);
+}
+
+FSkeletalControlReference USkeletalControlLibrary::SetAlpha(const FSkeletalControlReference& SkeletalControl, float Alpha)
+{
+	SkeletalControl.CallAnimNodeFunction<FAnimNode_SkeletalControlBase>(
+		TEXT("SetAlpha"),
+		[Alpha](FAnimNode_SkeletalControlBase& InSkeletalControl)
+		{
+			InSkeletalControl.SetAlpha(Alpha);
+		});
+
+	return SkeletalControl;
+}
+
+float USkeletalControlLibrary::GetAlpha(const FSkeletalControlReference& SkeletalControl)
+{
+	float Alpha = 0.0f;
+	
+	SkeletalControl.CallAnimNodeFunction<FAnimNode_SkeletalControlBase>(
+	TEXT("GetAlpha"),
+	[&Alpha](FAnimNode_SkeletalControlBase& InSkeletalControl)
+	{
+		Alpha = InSkeletalControl.GetAlpha();
+	});
+
+	return Alpha;
+}
\ No newline at end of file
diff --git a/Engine/Source/Runtime/AnimGraphRuntime/Public/AnimExecutionContextLibrary.h b/Engine/Source/Runtime/AnimGraphRuntime/Public/AnimExecutionContextLibrary.h
--- a/Engine/Source/Runtime/AnimGraphRuntime/Public/AnimExecutionContextLibrary.h
+++ b/Engine/Source/Runtime/AnimGraphRuntime/Public/AnimExecutionContextLibrary.h
@@ -27,8 +27,8 @@ public:
 	void Prototype_ThreadSafeAnimUpdateCall(const FAnimUpdateContext& Context, const FAnimNodeReference& Node) {}	
 #endif
 	
-	/** Get the anim instance that hosts this anim node */
-	UFUNCTION(BlueprintPure, Category = "Anim Node", meta=(BlueprintThreadSafe))
+	/** Get the anim instance that hosts this context */
+	UFUNCTION(BlueprintPure, Category = "Execution Context", meta=(BlueprintThreadSafe))
 	static UAnimInstance* GetAnimInstance(const FAnimExecutionContext& Context);
 
 	/** Internal compiler use only - Get a reference to an anim node by index */
@@ -36,18 +36,26 @@ public:
 	static FAnimNodeReference GetAnimNodeReference(UAnimInstance* Instance, int32 Index);
 
 	/** Convert to an initialization context */
-	UFUNCTION(BlueprintCallable, Category = "Anim Node", meta=(BlueprintThreadSafe, ExpandEnumAsExecs = "Result"))
-	FAnimInitializationContext ConvertToInitializationContext(const FAnimExecutionContext& Context, EAnimExecutionContextConversionResult& Result);
+	UFUNCTION(BlueprintCallable, Category = "Initialization Context", meta=(BlueprintThreadSafe, ExpandEnumAsExecs = "Result"))
+	static FAnimInitializationContext ConvertToInitializationContext(const FAnimExecutionContext& Context, EAnimExecutionContextConversionResult& Result);
 	
 	/** Convert to an update context */
-	UFUNCTION(BlueprintCallable, Category = "Anim Node", meta=(BlueprintThreadSafe, ExpandEnumAsExecs = "Result"))
+	UFUNCTION(BlueprintCallable, Category = "Update Context", meta=(BlueprintThreadSafe, ExpandEnumAsExecs = "Result"))
 	static FAnimUpdateContext ConvertToUpdateContext(const FAnimExecutionContext& Context, EAnimExecutionContextConversionResult& Result);
 
+	/** Get the current delta time in seconds */
+	UFUNCTION(BlueprintPure, Category = "Update Context", meta=(BlueprintThreadSafe))
+	static float GetDeltaTime(const FAnimUpdateContext& Context);
+
+	/** Get the current weight of this branch of the graph */
+	UFUNCTION(BlueprintPure, Category = "Update Context", meta=(BlueprintThreadSafe))
+	static float GetCurrentWeight(const FAnimUpdateContext& Context);
+	
 	/** Convert to a pose context */
-	UFUNCTION(BlueprintCallable, Category = "Anim Node", meta=(BlueprintThreadSafe, ExpandEnumAsExecs = "Result"))
+	UFUNCTION(BlueprintCallable, Category = "Pose Context", meta=(BlueprintThreadSafe, ExpandEnumAsExecs = "Result"))
 	static FAnimPoseContext ConvertToPoseContext(const FAnimExecutionContext& Context, EAnimExecutionContextConversionResult& Result);
 
 	/** Convert to a component space pose context */
-	UFUNCTION(BlueprintCallable, Category = "Anim Node", meta=(BlueprintThreadSafe, ExpandEnumAsExecs = "Result"))
+	UFUNCTION(BlueprintCallable, Category = "Component Space Pose Context", meta=(BlueprintThreadSafe, ExpandEnumAsExecs = "Result"))
 	static FAnimComponentSpacePoseContext ConvertToComponentSpacePoseContext(const FAnimExecutionContext& Context, EAnimExecutionContextConversionResult& Result);
 };
\ No newline at end of file
diff --git a/Engine/Source/Runtime/AnimGraphRuntime/Public/BoneControllers/AnimNode_SkeletalControlBase.h b/Engine/Source/Runtime/AnimGraphRuntime/Public/BoneControllers/AnimNode_SkeletalControlBase.h
--- a/Engine/Source/Runtime/AnimGraphRuntime/Public/BoneControllers/AnimNode_SkeletalControlBase.h
+++ b/Engine/Source/Runtime/AnimGraphRuntime/Public/BoneControllers/AnimNode_SkeletalControlBase.h
@@ -85,6 +85,12 @@ public:
 	virtual int32 GetLODThreshold() const override { return LODThreshold; }
 	// End of FAnimNode_Base interface
 
+	// Set the alpha of this node
+	void SetAlpha(float InAlpha);
+
+	// Get the alpha of this node
+	float GetAlpha() const;
+	
 protected:
 	// Interface for derived skeletal controls to implement
 	// use this function to update for skeletal control base
diff --git a/Engine/Source/Runtime/AnimGraphRuntime/Public/SkeletalControlLibrary.h b/Engine/Source/Runtime/AnimGraphRuntime/Public/SkeletalControlLibrary.h
new file mode 100644
--- /dev/null
+++ b/Engine/Source/Runtime/AnimGraphRuntime/Public/SkeletalControlLibrary.h
@@ -0,0 +1,49 @@
+// Copyright Epic Games, Inc. All Rights Reserved.
+
+#pragma once
+
+#include "CoreMinimal.h"
+
+#include "Kismet/BlueprintFunctionLibrary.h"
+#include "Animation/AnimNodeReference.h"
+#include "SkeletalControlLibrary.generated.h"
+
+struct FAnimNode_SkeletalControlBase;
+
+USTRUCT(BlueprintType)
+struct FSkeletalControlReference : public FAnimNodeReference
+{
+	GENERATED_BODY()
+
+	typedef FAnimNode_SkeletalControlBase FInternalNodeType;
+};
+
+// Exposes operations to be performed on a skeletal control anim node
+// Note: Experimental and subject to change!
+UCLASS(Experimental)
+class ANIMGRAPHRUNTIME_API USkeletalControlLibrary : public UBlueprintFunctionLibrary
+{
+	GENERATED_BODY()
+
+public:
+	/** Get a skeletal control from an anim node */
+	UFUNCTION(BlueprintCallable, Category = "Skeletal Control", meta=(BlueprintThreadSafe, ExpandEnumAsExecs = "Result"))
+	static FSkeletalControlReference ConvertToSkeletalControl(const FAnimNodeReference& Node, EAnimNodeReferenceConversionResult& Result);
+
+	/** Get a skeletal control from an anim node (pure) */
+	UFUNCTION(BlueprintPure, Category = "Skeletal Control", meta=(BlueprintThreadSafe, DisplayName = "Convert to Skeletal Control"))
+	static void ConvertToSkeletalControlPure(const FAnimNodeReference& Node, FSkeletalControlReference& SkeletalControl, bool& Result)
+	{
+		EAnimNodeReferenceConversionResult ConversionResult;
+		SkeletalControl = ConvertToSkeletalControl(Node, ConversionResult);
+		Result = (ConversionResult == EAnimNodeReferenceConversionResult::Succeeded);
+	}
+	
+	/** Set the alpha value of this skeletal control */
+	UFUNCTION(BlueprintCallable, Category = "Skeletal Control", meta=(BlueprintThreadSafe))
+	static FSkeletalControlReference SetAlpha(const FSkeletalControlReference& SkeletalControl, float Alpha);
+
+	/** Get the alpha value of this skeletal control */
+	UFUNCTION(BlueprintPure, Category = "Skeletal Control", meta=(BlueprintThreadSafe))
+	static float GetAlpha(const FSkeletalControlReference& SkeletalControl);
+};
\ No newline at end of file
diff --git a/Engine/Source/Runtime/Engine/Classes/Animation/AnimNodeReference.h b/Engine/Source/Runtime/Engine/Classes/Animation/AnimNodeReference.h
--- a/Engine/Source/Runtime/Engine/Classes/Animation/AnimNodeReference.h
+++ b/Engine/Source/Runtime/Engine/Classes/Animation/AnimNodeReference.h
@@ -37,7 +37,7 @@ public:
 	template<typename NodeType>
 	NodeType* GetAnimNodePtr() const
 	{
-		if(AnimNode && AnimNodeStruct && NodeType::StaticStruct()->IsChildOf(AnimNodeStruct))
+		if(AnimNode && AnimNodeStruct && AnimNodeStruct->IsChildOf(NodeType::StaticStruct()))
 		{
 			return static_cast<NodeType*>(AnimNode);
 		}
@@ -75,7 +75,7 @@ public:
 	{
 		static_assert(TIsDerivedFrom<OtherContextType, FAnimNodeReference>::IsDerived, "Argument ContextType must derive from FAnimNodeReference");
 		
-		if(InReference.AnimNodeStruct && OtherContextType::FInternalNodeType::StaticStruct()->IsChildOf(InReference.AnimNodeStruct))
+		if(InReference.AnimNodeStruct && InReference.AnimNodeStruct->IsChildOf(OtherContextType::FInternalNodeType::StaticStruct()))
 		{
 			OutResult = EAnimNodeReferenceConversionResult::Succeeded;
 			
