Fixed ensure/crash when using details panel to edit multi way blend weights

diff --git a/Engine/Source/Runtime/AnimGraphRuntime/Public/AnimNodes/AnimNode_MultiWayBlend.h b/Engine/Source/Runtime/AnimGraphRuntime/Public/AnimNodes/AnimNode_MultiWayBlend.h
--- a/Engine/Source/Runtime/AnimGraphRuntime/Public/AnimNodes/AnimNode_MultiWayBlend.h
+++ b/Engine/Source/Runtime/AnimGraphRuntime/Public/AnimNodes/AnimNode_MultiWayBlend.h
@@ -17,7 +17,7 @@ public:
 	UPROPERTY(EditAnywhere, BlueprintReadWrite, Category = Links)
 	TArray<FPoseLink> Poses;
 
-	UPROPERTY(EditAnywhere, BlueprintReadWrite, Category=Settings, meta=(PinShownByDefault))
+	UPROPERTY(EditAnywhere, EditFixedSize, BlueprintReadWrite, Category=Settings, meta=(PinShownByDefault))
 	TArray<float> DesiredAlphas;
 
 private:
