Chaos: Fix for crash when attempting to initialise a rigid body anim node without a valid mesh or skeleton.

diff --git a/Engine/Source/Runtime/AnimGraphRuntime/Private/BoneControllers/AnimNode_RigidBody.cpp b/Engine/Source/Runtime/AnimGraphRuntime/Private/BoneControllers/AnimNode_RigidBody.cpp
--- a/Engine/Source/Runtime/AnimGraphRuntime/Private/BoneControllers/AnimNode_RigidBody.cpp
+++ b/Engine/Source/Runtime/AnimGraphRuntime/Private/BoneControllers/AnimNode_RigidBody.cpp
@@ -1014,11 +1014,18 @@ void FAnimNode_RigidBody::InitPhysics(const UAnimInstance* InAnimInstance)
 
 	const USkeletalMeshComponent* SkeletalMeshComp = InAnimInstance->GetSkelMeshComponent();
 	const USkeletalMesh* SkeletalMeshAsset = SkeletalMeshComp->SkeletalMesh;
+	USkeleton* SkeletonAsset = InAnimInstance->CurrentSkeleton;
 
 	const FReferenceSkeleton& SkelMeshRefSkel = SkeletalMeshAsset->GetRefSkeleton();
 	UsePhysicsAsset = OverridePhysicsAsset ? ToRawPtr(OverridePhysicsAsset) : InAnimInstance->GetSkelMeshComponent()->GetPhysicsAsset();
-		
-	USkeleton* SkeletonAsset = InAnimInstance->CurrentSkeleton;
+	
+	if(!SkeletalMeshAsset || !SkeletonAsset)
+	{
+		// Without both the skeleton and the mesh we can't create a new simulation.
+		// The previous simulation has just been cleaned up above so we can return early here and not instantiate a new one
+		return;
+	}
+
 	ensure(SkeletonAsset == SkeletalMeshAsset->GetSkeleton());
 
 	const int32 SkelMeshLinkupIndex = SkeletonAsset->GetMeshLinkupIndex(SkeletalMeshAsset);
