Fix crash ending PIE when a sequence is playing that is driving an anim BP with a spring bone controller active

diff --git a/Engine/Source/Runtime/AnimGraphRuntime/Private/BoneControllers/AnimNode_SpringBone.cpp b/Engine/Source/Runtime/AnimGraphRuntime/Private/BoneControllers/AnimNode_SpringBone.cpp
--- a/Engine/Source/Runtime/AnimGraphRuntime/Private/BoneControllers/AnimNode_SpringBone.cpp
+++ b/Engine/Source/Runtime/AnimGraphRuntime/Private/BoneControllers/AnimNode_SpringBone.cpp
@@ -232,19 +232,23 @@ void FAnimNode_SpringBone::InitializeBoneReferences(const FBoneContainer& Requir
 
 void FAnimNode_SpringBone::PreUpdate(const UAnimInstance* InAnimInstance)
 {
-	const USkeletalMeshComponent* SkelComp = InAnimInstance->GetSkelMeshComponent();
-	const UWorld* World = SkelComp->GetWorld();
-	check(World->GetWorldSettings());
-	TimeDilation = World->GetWorldSettings()->GetEffectiveTimeDilation();
-
-	AActor* SkelOwner = SkelComp->GetOwner();
-	if (SkelComp->GetAttachParent() != NULL && (SkelOwner == NULL))
+	if (const USkeletalMeshComponent* SkelComp = InAnimInstance->GetSkelMeshComponent())
 	{
-		SkelOwner = SkelComp->GetAttachParent()->GetOwner();
-		OwnerVelocity = SkelOwner->GetVelocity();
+		if (const UWorld* World = SkelComp->GetWorld())
+		{
+			check(World->GetWorldSettings());
+			TimeDilation = World->GetWorldSettings()->GetEffectiveTimeDilation();
+
+			AActor* SkelOwner = SkelComp->GetOwner();
+			if (SkelComp->GetAttachParent() != nullptr && (SkelOwner == nullptr))
+			{
+				SkelOwner = SkelComp->GetAttachParent()->GetOwner();
+				OwnerVelocity = SkelOwner->GetVelocity();
+			}
+			else
+			{
+				OwnerVelocity = FVector::ZeroVector;
+			}
+		}
 	}
-	else
-	{
-		OwnerVelocity = FVector::ZeroVector;
-	}	
 }
