Fixes for Mirror Node
  - remove optimization that can introduce bugs in circumstances with identical joint counts but diferent indexes
  - added constant folding.

diff --git a/Engine/Source/Runtime/AnimGraphRuntime/Private/AnimNodes/AnimNode_Mirror.cpp b/Engine/Source/Runtime/AnimGraphRuntime/Private/AnimNodes/AnimNode_Mirror.cpp
--- a/Engine/Source/Runtime/AnimGraphRuntime/Private/AnimNodes/AnimNode_Mirror.cpp
+++ b/Engine/Source/Runtime/AnimGraphRuntime/Private/AnimNodes/AnimNode_Mirror.cpp
@@ -9,36 +9,99 @@
 
 #define LOCTEXT_NAMESPACE "AnimNode_Mirror"
 
-FAnimNode_Mirror::FAnimNode_Mirror()
-	: bMirror(true)
-	, MirrorDataTable(nullptr)
-	, BlendTimeOnMirrorStateChange(0.0f)
-	, bBoneMirroring(true)
-	, bCurveMirroring(true)
-	, bAttributeMirroring(true)
-	, bResetChildOnMirrorStateChange(false)
-	, bMirrorState(false)
+FAnimNode_MirrorBase::FAnimNode_MirrorBase()
+	: bMirrorState(false)
 	, bMirrorStateIsValid(false)
 {
 }
 
-UMirrorDataTable* FAnimNode_Mirror::GetMirrorDataTable() const 
-{ 
-	return MirrorDataTable.Get();
+UMirrorDataTable* FAnimNode_MirrorBase::GetMirrorDataTable() const
+{
+	return nullptr;
+}
+
+bool FAnimNode_MirrorBase::SetMirrorDataTable(UMirrorDataTable* MirrorTable)
+{
+	return false;
+}
+
+bool FAnimNode_MirrorBase::GetMirror() const
+{
+	return false;
+}
+
+float FAnimNode_MirrorBase::GetBlendTimeOnMirrorStateChange() const
+{
+	return 0.0;
+}
+
+bool FAnimNode_MirrorBase::GetBoneMirroring() const
+{
+	return false;
+}
+
+bool FAnimNode_MirrorBase::GetCurveMirroring() const
+{
+	return false;
+}
+
+bool FAnimNode_MirrorBase::GetAttributeMirroring() const
+{
+	return false;
+}
+
+bool FAnimNode_MirrorBase::GetResetChildOnMirrorStateChange() const
+{
+	return false;
+}
+
+bool FAnimNode_MirrorBase::SetMirror(bool bInMirror)
+{
+	return false;
+}
+
+bool FAnimNode_MirrorBase::SetBlendTimeOnMirrorStateChange(float InBlendTime)
+{
+	return false;
+}
+
+bool FAnimNode_MirrorBase::SetBoneMirroring(bool bInBoneMirroring)
+{
+	return false;
+}
+
+bool FAnimNode_MirrorBase::SetCurveMirroring(bool bInCurveMirroring)
+{
+	return false;
+}
+
+bool FAnimNode_MirrorBase::SetAttributeMirroring(bool bInAttributeMirroring)
+{
+	return false;
+}
+
+bool FAnimNode_MirrorBase::SetResetChildOnMirrorStateChange(bool bInResetChildOnMirrorStateChange)
+{
+	return false;
 }
 
-void FAnimNode_Mirror::SetMirrorDataTable(UMirrorDataTable* MirrorTable)
+void FAnimNode_MirrorBase::SetSourceLinkNode(FAnimNode_Base* NewLinkNode)
 {
-	MirrorDataTable = MirrorTable; 
+	Source.SetLinkNode(NewLinkNode);
 }
 
-void FAnimNode_Mirror::Initialize_AnyThread(const FAnimationInitializeContext& Context)
+FAnimNode_Base* FAnimNode_MirrorBase::GetSourceLinkNode()
+{
+	return Source.GetLinkNode(); 
+}
+
+void FAnimNode_MirrorBase::Initialize_AnyThread(const FAnimationInitializeContext& Context)
 {
 	FAnimNode_Base::Initialize_AnyThread(Context);
 	Source.Initialize(Context);
 }
 
-void FAnimNode_Mirror::CacheBones_AnyThread(const FAnimationCacheBonesContext& Context)
+void FAnimNode_MirrorBase::CacheBones_AnyThread(const FAnimationCacheBonesContext& Context)
 {
 	Source.CacheBones(Context);
 	DECLARE_SCOPE_HIERARCHICAL_COUNTER_ANIMNODE(CacheBones_AnyThread)
@@ -47,23 +110,14 @@ void FAnimNode_Mirror::CacheBones_AnyThread(const FAnimationCacheBonesContext& C
 	FillCompactPoseAndComponentRefRotations(BoneContainer);
 }
 
-void FAnimNode_Mirror::FillCompactPoseAndComponentRefRotations(const FBoneContainer& BoneContainer)
+void FAnimNode_MirrorBase::FillCompactPoseAndComponentRefRotations(const FBoneContainer& BoneContainer)
 {
+	UMirrorDataTable* MirrorDataTable = GetMirrorDataTable();
 	if (MirrorDataTable)
 	{
-		bool bSharesSkeleton = BoneContainer.GetSkeletonAsset() == MirrorDataTable->Skeleton;
-		bool bIdenticalJointCount = BoneContainer.GetReferenceSkeleton().GetNum() == MirrorDataTable->BoneToMirrorBoneIndex.Num();
-
-		if (bSharesSkeleton && bIdenticalJointCount)
-		{
-			MirrorDataTable->FillCompactPoseMirrorBones(BoneContainer, MirrorDataTable->BoneToMirrorBoneIndex, CompactPoseMirrorBones);
-		}
-		else
-		{
-			TArray<int32> MirrorBoneIndexes;
-			MirrorDataTable->FillMirrorBoneIndexes(BoneContainer.GetReferenceSkeleton(), MirrorBoneIndexes);
-			MirrorDataTable->FillCompactPoseMirrorBones(BoneContainer, MirrorBoneIndexes, CompactPoseMirrorBones);
-		}
+		TArray<int32> MirrorBoneIndexes;
+		MirrorDataTable->FillMirrorBoneIndexes(BoneContainer.GetReferenceSkeleton(), MirrorBoneIndexes);
+		MirrorDataTable->FillCompactPoseMirrorBones(BoneContainer, MirrorBoneIndexes, CompactPoseMirrorBones);
 
 		const int32 NumBones = BoneContainer.GetCompactPoseNumBones();
 		ComponentSpaceRefRotations.SetNumUninitialized(NumBones);
@@ -81,19 +135,20 @@ void FAnimNode_Mirror::FillCompactPoseAndComponentRefRotations(const FBoneContai
 	}
 }
 
-void FAnimNode_Mirror::Update_AnyThread(const FAnimationUpdateContext& Context)
+void FAnimNode_MirrorBase::Update_AnyThread(const FAnimationUpdateContext& Context)
 {
 	GetEvaluateGraphExposedInputs().Execute(Context);
 
+	bool bMirror = GetMirror();
 	if (bMirrorStateIsValid && bMirrorState != bMirror)
 	{
-		if (BlendTimeOnMirrorStateChange > SMALL_NUMBER)
+		if (GetBlendTimeOnMirrorStateChange() > SMALL_NUMBER)
 		{
 			// Inertialize when switching between mirrored and unmirrored states to smooth out the pose discontinuity
 			UE::Anim::IInertializationRequester* InertializationRequester = Context.GetMessage<UE::Anim::IInertializationRequester>();
 			if (InertializationRequester)
 			{
-				InertializationRequester->RequestInertialization(BlendTimeOnMirrorStateChange);
+				InertializationRequester->RequestInertialization(GetBlendTimeOnMirrorStateChange());
 				InertializationRequester->AddDebugRecord(*Context.AnimInstanceProxy, Context.GetCurrentNodeId());
 			}
 			else
@@ -103,28 +158,29 @@ void FAnimNode_Mirror::Update_AnyThread(const FAnimationUpdateContext& Context)
 		}
 
 		// Optionally reinitialize the source when the mirror state changes
-		if (bResetChildOnMirrorStateChange)
+		if (GetResetChildOnMirrorStateChange())
 		{
 			FAnimationInitializeContext ReinitializeContext(Context.AnimInstanceProxy, Context.SharedContext);
 			Source.Initialize(ReinitializeContext);
 		}
 	}
-
+	UMirrorDataTable* MirrorDataTable = GetMirrorDataTable();
 	UE::Anim::TOptionalScopedGraphMessage<UE::Anim::FMirrorSyncScope> Message(bMirror, Context, Context, MirrorDataTable);
-	
+
 	bMirrorState = bMirror;
 	bMirrorStateIsValid = true;
 
 	Source.Update(Context);
 }
 
-void FAnimNode_Mirror::Evaluate_AnyThread(FPoseContext& Output)
+void FAnimNode_MirrorBase::Evaluate_AnyThread(FPoseContext& Output)
 {
 	Source.Evaluate(Output);
+	UMirrorDataTable* MirrorDataTable = GetMirrorDataTable();
 
 	if (bMirrorState && MirrorDataTable)
 	{
-		if (bBoneMirroring)
+		if (GetBoneMirroring())
 		{
 			const FBoneContainer& BoneContainer = Output.Pose.GetBoneContainer();
 			const TArray<FBoneIndexType>& RequiredBoneIndices = BoneContainer.GetBoneIndicesArray();
@@ -136,19 +192,19 @@ void FAnimNode_Mirror::Evaluate_AnyThread(FPoseContext& Output)
 			FAnimationRuntime::MirrorPose(Output.Pose, MirrorDataTable->MirrorAxis, CompactPoseMirrorBones, ComponentSpaceRefRotations);
 		}
 
-		if (bCurveMirroring)
+		if (GetCurveMirroring())
 		{
 			FAnimationRuntime::MirrorCurves(Output.Curve, *MirrorDataTable);
 		}
 
-		if (bAttributeMirroring)
+		if (GetAttributeMirroring())
 		{
 			UE::Anim::Attributes::MirrorAttributes(Output.CustomAttributes, *MirrorDataTable);
 		}
 	}
 }
 
-void FAnimNode_Mirror::GatherDebugData(FNodeDebugData& DebugData)
+void FAnimNode_MirrorBase::GatherDebugData(FNodeDebugData& DebugData)
 {
 	FString DebugLine = DebugData.GetNodeName(this);
 	DebugLine += FString::Printf(TEXT("(Mirrored: %s)"), (bMirrorState) ? TEXT("true") : TEXT("false"));
@@ -157,5 +213,230 @@ void FAnimNode_Mirror::GatherDebugData(FNodeDebugData& DebugData)
 	Source.GatherDebugData(DebugData);
 }
 
+FAnimNode_Mirror::FAnimNode_Mirror()
+{
+}
+
+UMirrorDataTable* FAnimNode_Mirror::GetMirrorDataTable() const 
+{ 
+	return GET_ANIM_NODE_DATA(TObjectPtr<UMirrorDataTable>, MirrorDataTable);
+}
+
+bool FAnimNode_Mirror::SetMirrorDataTable(UMirrorDataTable* InMirrorTable)
+{
+#if WITH_EDITORONLY_DATA
+	MirrorDataTable = InMirrorTable;
+	GET_MUTABLE_ANIM_NODE_DATA(TObjectPtr<UMirrorDataTable>, MirrorDataTable) = InMirrorTable;
+#endif
+
+	if (TObjectPtr<UMirrorDataTable>* InMirrorTablePtr = GET_INSTANCE_ANIM_NODE_DATA_PTR(TObjectPtr<UMirrorDataTable>, MirrorDataTable))
+	{
+		*InMirrorTablePtr = InMirrorTable;
+		return true;
+	}
+
+	return false;
+}
+
+bool FAnimNode_Mirror::GetMirror() const
+{
+	return GET_ANIM_NODE_DATA(bool, bMirror);
+}
+
+float FAnimNode_Mirror::GetBlendTimeOnMirrorStateChange() const
+{
+	return GET_ANIM_NODE_DATA(float, BlendTimeOnMirrorStateChange);
+}
+
+bool FAnimNode_Mirror::GetBoneMirroring() const
+{
+	return GET_ANIM_NODE_DATA(bool, bBoneMirroring);
+}
+
+bool FAnimNode_Mirror::GetCurveMirroring() const
+{
+	return GET_ANIM_NODE_DATA(bool, bCurveMirroring);
+}
+
+bool FAnimNode_Mirror::GetAttributeMirroring() const
+{
+	return GET_ANIM_NODE_DATA(bool, bAttributeMirroring);
+}
+
+bool FAnimNode_Mirror::GetResetChildOnMirrorStateChange() const
+{
+	return GET_ANIM_NODE_DATA(bool, bResetChildOnMirrorStateChange);
+}
+
+bool FAnimNode_Mirror::SetMirror(bool bInMirror)
+{
+#if WITH_EDITORONLY_DATA
+	bMirror = bInMirror;
+#endif
+
+	if (bool* bMirrorPtr = GET_INSTANCE_ANIM_NODE_DATA_PTR(bool, bMirror))
+	{
+		*bMirrorPtr = bInMirror;
+		return true;
+	}
+	return false;
+}
+
+bool FAnimNode_Mirror::SetBlendTimeOnMirrorStateChange(float InBlendTime)
+{
+#if WITH_EDITORONLY_DATA
+	BlendTimeOnMirrorStateChange = InBlendTime;
+#endif
+
+	if (float* BlendTimePtr = GET_INSTANCE_ANIM_NODE_DATA_PTR(float, BlendTimeOnMirrorStateChange))
+	{
+		*BlendTimePtr = InBlendTime;
+		return true;
+	}
+	return false;
+}
+
+bool FAnimNode_Mirror::SetBoneMirroring(bool bInBoneMirroring)
+{
+#if WITH_EDITORONLY_DATA
+	bBoneMirroring = bInBoneMirroring;
+#endif
+
+	if (bool* bBoneMirroringPtr = GET_INSTANCE_ANIM_NODE_DATA_PTR(bool, bBoneMirroring))
+	{
+		*bBoneMirroringPtr = bInBoneMirroring;
+		return true;
+	}
+
+	return false;
+}
+
+bool FAnimNode_Mirror::SetCurveMirroring(bool bInCurveMirroring)
+{
+#if WITH_EDITORONLY_DATA
+	bCurveMirroring = bInCurveMirroring;
+#endif
+
+	if (bool* bCurveMirroringPtr = GET_INSTANCE_ANIM_NODE_DATA_PTR(bool, bCurveMirroring))
+	{
+		*bCurveMirroringPtr = bInCurveMirroring;
+		return true;
+	}
+
+	return false;
+}
+
+bool FAnimNode_Mirror::SetAttributeMirroring(bool bInAttributeMirroring)
+{
+#if WITH_EDITORONLY_DATA
+	bAttributeMirroring = bInAttributeMirroring;
+#endif
+
+	if (bool* bAttributeMirroringPtr = GET_INSTANCE_ANIM_NODE_DATA_PTR(bool, bAttributeMirroring))
+	{
+		*bAttributeMirroringPtr = bInAttributeMirroring;
+		return true;
+	}
+
+	return false;
+}
+
+bool FAnimNode_Mirror::SetResetChildOnMirrorStateChange(bool bInResetChildOnMirrorStateChange)
+{
+#if WITH_EDITORONLY_DATA
+	bResetChildOnMirrorStateChange = bInResetChildOnMirrorStateChange;
+#endif
+
+	if (bool* bResetChildOnMirrorStateChangePtr = GET_INSTANCE_ANIM_NODE_DATA_PTR(bool, bResetChildOnMirrorStateChange))
+	{
+		*bResetChildOnMirrorStateChangePtr = bInResetChildOnMirrorStateChange;
+		return true;
+	}
+
+	return false;
+}
+
+
+FAnimNode_Mirror_Standalone::FAnimNode_Mirror_Standalone()
+{
+}
+
+UMirrorDataTable* FAnimNode_Mirror_Standalone::GetMirrorDataTable() const
+{
+	return MirrorDataTable.Get();
+}
+
+bool FAnimNode_Mirror_Standalone::SetMirrorDataTable(UMirrorDataTable* MirrorTable)
+{
+	MirrorDataTable = MirrorTable;
+	return true;
+}
+
+bool FAnimNode_Mirror_Standalone::GetMirror() const
+{
+	return bMirror;
+}
+
+float FAnimNode_Mirror_Standalone::GetBlendTimeOnMirrorStateChange() const
+{
+	return BlendTimeOnMirrorStateChange;
+}
+
+bool FAnimNode_Mirror_Standalone::GetBoneMirroring() const
+{
+	return bBoneMirroring;
+}
+
+bool FAnimNode_Mirror_Standalone::GetCurveMirroring() const
+{
+	return bCurveMirroring;
+}
+
+bool FAnimNode_Mirror_Standalone::GetAttributeMirroring() const
+{
+	return bAttributeMirroring;
+}
+
+bool FAnimNode_Mirror_Standalone::GetResetChildOnMirrorStateChange() const
+{
+	return bResetChildOnMirrorStateChange;
+}
+
+bool FAnimNode_Mirror_Standalone::SetMirror(bool bInMirror)
+{
+	bMirror = bInMirror;
+	return true;
+}
+
+bool FAnimNode_Mirror_Standalone::SetBlendTimeOnMirrorStateChange(float InBlendTime)
+{
+	BlendTimeOnMirrorStateChange = InBlendTime;
+	return true;
+}
+
+bool FAnimNode_Mirror_Standalone::SetBoneMirroring(bool bInBoneMirroring)
+{
+	bBoneMirroring = bInBoneMirroring;
+	return true;
+
+}
+
+bool FAnimNode_Mirror_Standalone::SetCurveMirroring(bool bInCurveMirroring)
+{
+	bCurveMirroring = bInCurveMirroring;
+	return true;
+}
+
+bool FAnimNode_Mirror_Standalone::SetAttributeMirroring(bool bInAttributeMirroring)
+{
+	bAttributeMirroring = bInAttributeMirroring;
+	return true;
+}
+
+bool FAnimNode_Mirror_Standalone::SetResetChildOnMirrorStateChange(bool bInResetChildOnMirrorStateChange)
+{
+	bResetChildOnMirrorStateChange = bInResetChildOnMirrorStateChange;
+	return true;
+}
 
 #undef LOCTEXT_NAMESPACE
diff --git a/Engine/Source/Runtime/AnimGraphRuntime/Private/AnimSequencerInstanceProxy.cpp b/Engine/Source/Runtime/AnimGraphRuntime/Private/AnimSequencerInstanceProxy.cpp
--- a/Engine/Source/Runtime/AnimGraphRuntime/Private/AnimSequencerInstanceProxy.cpp
+++ b/Engine/Source/Runtime/AnimGraphRuntime/Private/AnimSequencerInstanceProxy.cpp
@@ -87,7 +87,7 @@ void FAnimSequencerInstanceProxy::ClearSequencePlayerAndMirrorMaps()
 
 	SequencerToPlayerMap.Empty();
 
-	for (TPair<uint32, FAnimNode_Mirror*>& Iter : SequencerToMirrorMap)
+	for (TPair<uint32, FAnimNode_Mirror_Standalone*>& Iter : SequencerToMirrorMap)
 	{
 		delete Iter.Value;
 	}
@@ -139,9 +139,9 @@ void FAnimSequencerInstanceProxy::InitAnimTrack(UAnimSequenceBase* InAnimSequenc
 			SequencerToPlayerMap.Add(SequenceId, NewPlayerState);
 
 			// link player to mirror node,
-			FAnimNode_Mirror* NewMirrorNode = new FAnimNode_Mirror();
-			NewMirrorNode->bMirror = false;
-			NewMirrorNode->Source.SetLinkNode(&NewPlayerState->PlayerNode);
+			FAnimNode_Mirror_Standalone* NewMirrorNode = new FAnimNode_Mirror_Standalone();
+			NewMirrorNode->SetMirror(false);
+			NewMirrorNode->SetSourceLinkNode(&NewPlayerState->PlayerNode);
 			SequencerToMirrorMap.Add(SequenceId, NewMirrorNode); 
 
 			// link mirror to blendnode, this will let you trigger notifies and so on
@@ -159,7 +159,7 @@ void FAnimSequencerInstanceProxy::InitAnimTrack(UAnimSequenceBase* InAnimSequenc
 		// initialize player
 		PlayerState->PlayerNode.Initialize_AnyThread(FAnimationInitializeContext(this));
 
-		FAnimNode_Mirror* Mirror = SequencerToMirrorMap.FindRef(SequenceId);
+		FAnimNode_Mirror_Standalone* Mirror = SequencerToMirrorMap.FindRef(SequenceId);
 		if (Mirror)
 		{
 			Mirror->Initialize_AnyThread(FAnimationInitializeContext(this));
@@ -227,10 +227,10 @@ void FAnimSequencerInstanceProxy::UpdateAnimTrack(UAnimSequenceBase* InAnimSeque
 		PlayerState->PlayerNode.SetReinitializationBehavior(ESequenceEvalReinit::NoReset);
 	}
 
-	FAnimNode_Mirror* MirrorNode = SequencerToMirrorMap.FindRef(SequenceId);
+	FAnimNode_Mirror_Standalone* MirrorNode = SequencerToMirrorMap.FindRef(SequenceId);
 	if (MirrorNode)
 	{
-		MirrorNode->bMirror = InMirrorDataTable != nullptr;
+		MirrorNode->SetMirror(InMirrorDataTable != nullptr);
 		UMirrorDataTable* OldMirrorDataTable = MirrorNode->GetMirrorDataTable();
 		MirrorNode->SetMirrorDataTable(InMirrorDataTable);
 
diff --git a/Engine/Source/Runtime/AnimGraphRuntime/Public/AnimNodes/AnimNode_Mirror.h b/Engine/Source/Runtime/AnimGraphRuntime/Public/AnimNodes/AnimNode_Mirror.h
--- a/Engine/Source/Runtime/AnimGraphRuntime/Public/AnimNodes/AnimNode_Mirror.h
+++ b/Engine/Source/Runtime/AnimGraphRuntime/Public/AnimNodes/AnimNode_Mirror.h
@@ -8,16 +8,12 @@
 #include "Animation/MirrorDataTable.h"
 #include "AnimNode_Mirror.generated.h"
 
-
 USTRUCT(BlueprintInternalUseOnly)
-struct ANIMGRAPHRUNTIME_API FAnimNode_Mirror : public FAnimNode_Base
+struct ANIMGRAPHRUNTIME_API FAnimNode_MirrorBase : public FAnimNode_Base
 {
 	GENERATED_BODY()
-
-	friend class UAnimGraphNode_Mirror;
-
 public:
-	FAnimNode_Mirror();
+	FAnimNode_MirrorBase(); 
 
 	virtual void Initialize_AnyThread(const FAnimationInitializeContext& Context) override;
 	virtual void CacheBones_AnyThread(const FAnimationCacheBonesContext& Context) override;
@@ -25,36 +21,61 @@ public:
 	virtual void Evaluate_AnyThread(FPoseContext& Output) override;
 	virtual void GatherDebugData(FNodeDebugData& DebugData) override;
 
-	float GetBlendTimeOnMirrorStateChange() const { return BlendTimeOnMirrorStateChange; }
+	// Get the MirrorDataTable
+	virtual UMirrorDataTable* GetMirrorDataTable() const;
 
-	UMirrorDataTable* GetMirrorDataTable() const;
-	void SetMirrorDataTable(UMirrorDataTable* MirrorTable);
+	// Set the MirrorDataTable
+	virtual bool SetMirrorDataTable(UMirrorDataTable* MirrorTable);
 
-	UPROPERTY(EditAnywhere, Category = Links)
-	FPoseLink Source;
+	// Get Mirror State
+	virtual bool GetMirror() const;
 
-	UPROPERTY(EditAnywhere, Category = Settings, meta = (PinShownByDefault))
-	bool bMirror;
-private:
+	// How long to blend using inertialization when switching  mirrored state
+	virtual float GetBlendTimeOnMirrorStateChange() const;
 
-	UPROPERTY(EditAnywhere, Category = Settings)
-	TObjectPtr<UMirrorDataTable> MirrorDataTable;
+	// Should bones mirror
+	virtual bool GetBoneMirroring() const;
 
-	UPROPERTY(EditAnywhere, Category = Settings, meta = (PinHiddenByDefault))
-	float BlendTimeOnMirrorStateChange;
+	// Should the curves mirror
+	virtual bool GetCurveMirroring() const;
 
-	UPROPERTY(EditAnywhere, Category = MirroredChannels, meta=(DisplayName="Bone"))
-	bool bBoneMirroring;
+	// Should attributes mirror (based on the bone mirroring data in the mirror data table) 
+	virtual bool GetAttributeMirroring() const;
 
-	UPROPERTY(EditAnywhere, Category = MirroredChannels, meta=(DisplayName = "Curve"))
-	bool bCurveMirroring;
+	// Whether to reset (reinitialize) the child (source) pose when the mirror state changes
+	virtual bool GetResetChildOnMirrorStateChange() const;
 
-	UPROPERTY(EditAnywhere, Category = MirroredChannels, meta=(DisplayName = "Attributes"))
-	bool bAttributeMirroring;
+	// Set Mirror State
+	virtual bool SetMirror(bool bInMirror);
 
-	// Whether to reset (reinitialize) the child (source) pose when the mirror state changes
-	UPROPERTY(EditAnywhere, Category = Settings)
-	bool bResetChildOnMirrorStateChange;
+	// Set how long to blend using inertialization when switching  mirrored state
+	// @return true if the value was set (it is dynamic), or false if it could not (it is not dynamic or pin exposed)
+	virtual bool SetBlendTimeOnMirrorStateChange(float InBlendTime);
+
+	// Set if bones mirror
+	// @return true if the value was set (it is dynamic), or false if it could not (it is not dynamic or pin exposed)
+	virtual bool SetBoneMirroring(bool bInBoneMirroring);
+
+	// Set if curves mirror
+	// @return true if the value was set (it is dynamic), or false if it could not (it is not dynamic or pin exposed)
+	virtual bool SetCurveMirroring(bool bInCurveMirroring);
+
+	// Set if attributes mirror
+	// @return true if the value was set (it is dynamic), or false if it could not (it is not dynamic or pin exposed)
+	virtual bool SetAttributeMirroring(bool bInAttributeMirroring);
+
+	// Set whether to reset (reinitialize) the child (source) pose when the mirror state changes
+	// @return true if the value was set (it is dynamic), or false if it could not (it is not dynamic or pin exposed)
+	virtual bool SetResetChildOnMirrorStateChange(bool bInResetChildOnMirrorStateChange);
+
+	/** This only used by custom handlers, and it is advanced feature. */
+	virtual void SetSourceLinkNode(FAnimNode_Base* NewLinkNode);
+
+	/** This only used by custom handlers, and it is advanced feature. */
+	virtual FAnimNode_Base* GetSourceLinkNode();
+protected:
+	UPROPERTY(EditAnywhere, Category = Links)
+	FPoseLink Source;
 
 private:
 	bool bMirrorState;
@@ -67,3 +88,111 @@ private:
 	// Pre-calculated component space of reference pose, which allows mirror to work with any joint orient 
 	TCustomBoneIndexArray<FQuat, FCompactPoseBoneIndex> ComponentSpaceRefRotations;
 };
+
+
+USTRUCT(BlueprintInternalUseOnly)
+struct ANIMGRAPHRUNTIME_API FAnimNode_Mirror : public FAnimNode_MirrorBase
+{
+	GENERATED_BODY()
+
+	friend class UAnimGraphNode_Mirror;
+
+public:
+	FAnimNode_Mirror();
+
+	virtual UMirrorDataTable* GetMirrorDataTable() const override;
+	virtual bool SetMirrorDataTable(UMirrorDataTable* MirrorTable) override;
+
+	virtual bool GetMirror() const override;
+	virtual float GetBlendTimeOnMirrorStateChange() const override;
+	virtual bool GetBoneMirroring() const override;
+	virtual bool GetCurveMirroring() const override;
+	virtual bool GetAttributeMirroring() const override;
+	virtual bool GetResetChildOnMirrorStateChange() const override;
+
+	virtual bool SetMirror(bool bInMirror) override;
+	virtual bool SetBlendTimeOnMirrorStateChange(float InBlendTime) override;
+	virtual bool SetBoneMirroring(bool bInBoneMirroring) override;
+	virtual bool SetCurveMirroring(bool bInCurveMirroring) override;
+	virtual bool SetAttributeMirroring(bool bInAttributeMirroring) override;
+	virtual bool SetResetChildOnMirrorStateChange(bool bInResetChildOnMirrorStateChange) override;
+
+protected:
+
+#if WITH_EDITORONLY_DATA
+	UPROPERTY(EditAnywhere, Category = Settings, meta = (PinShownByDefault, FoldProperty))
+	bool bMirror = true; 
+
+	UPROPERTY(EditAnywhere, Category = Settings, meta = (FoldProperty))
+	TObjectPtr<UMirrorDataTable> MirrorDataTable = nullptr;
+
+	UPROPERTY(EditAnywhere, Category = Settings, meta = (PinHiddenByDefault, FoldProperty))
+	float BlendTimeOnMirrorStateChange = 0.0f;
+
+	UPROPERTY(EditAnywhere, Category = MirroredChannels, meta=(DisplayName="Bone", FoldProperty))
+	bool bBoneMirroring = true;
+
+	UPROPERTY(EditAnywhere, Category = MirroredChannels, meta=(DisplayName = "Curve", FoldProperty))
+	bool bCurveMirroring = true;
+
+	UPROPERTY(EditAnywhere, Category = MirroredChannels, meta=(DisplayName = "Attributes", FoldProperty))
+	bool bAttributeMirroring = true;
+
+	// Whether to reset (reinitialize) the child (source)pose when the mirror state changes
+	UPROPERTY(EditAnywhere, Category = Settings, meta = ( FoldProperty))
+	bool bResetChildOnMirrorStateChange = false;
+#endif
+};
+
+
+USTRUCT(BlueprintInternalUseOnly)
+struct ANIMGRAPHRUNTIME_API FAnimNode_Mirror_Standalone : public FAnimNode_MirrorBase
+{
+	GENERATED_BODY()
+
+	friend class UAnimGraphNode_Mirror;
+
+public:
+	FAnimNode_Mirror_Standalone();
+
+	virtual UMirrorDataTable* GetMirrorDataTable() const override;
+	virtual bool SetMirrorDataTable(UMirrorDataTable* MirrorTable) override;
+
+	virtual bool GetMirror() const override;
+	virtual float GetBlendTimeOnMirrorStateChange() const override;
+	virtual bool GetBoneMirroring() const override;
+	virtual bool GetCurveMirroring() const override;
+	virtual bool GetAttributeMirroring() const override;
+	virtual bool GetResetChildOnMirrorStateChange() const override;
+
+	virtual bool SetMirror(bool bInMirror) override;
+	virtual bool SetBlendTimeOnMirrorStateChange(float InBlendTime) override;
+	virtual bool SetBoneMirroring(bool bInBoneMirroring) override;
+	virtual bool SetCurveMirroring(bool bInCurveMirroring) override;
+	virtual bool SetAttributeMirroring(bool bInAttributeMirroring) override;
+	virtual bool SetResetChildOnMirrorStateChange(bool bInResetChildOnMirrorStateChange) override;
+
+protected:
+
+	UPROPERTY(EditAnywhere, Category = Settings, meta = (PinShownByDefault))
+	bool bMirror = true;
+
+	UPROPERTY(EditAnywhere, Category = Settings)
+	TObjectPtr<UMirrorDataTable> MirrorDataTable = nullptr;
+
+	UPROPERTY(EditAnywhere, Category = Settings, meta = (PinHiddenByDefault, FoldProperty))
+	float BlendTimeOnMirrorStateChange = 0.0f;
+
+	UPROPERTY(EditAnywhere, Category = MirroredChannels, meta = (DisplayName = "Bone", FoldProperty))
+	bool bBoneMirroring = true;
+
+	UPROPERTY(EditAnywhere, Category = MirroredChannels, meta = (DisplayName = "Curve", FoldProperty))
+	bool bCurveMirroring = true;
+
+	UPROPERTY(EditAnywhere, Category = MirroredChannels, meta = (DisplayName = "Attributes", FoldProperty))
+	bool bAttributeMirroring = true;
+
+	// Whether to reset (reinitialize) the child (source)pose when the mirror state changes
+	UPROPERTY(EditAnywhere, Category = Settings, meta = (FoldProperty))
+	bool bResetChildOnMirrorStateChange = false;
+};
\ No newline at end of file
diff --git a/Engine/Source/Runtime/AnimGraphRuntime/Public/AnimSequencerInstanceProxy.h b/Engine/Source/Runtime/AnimGraphRuntime/Public/AnimSequencerInstanceProxy.h
--- a/Engine/Source/Runtime/AnimGraphRuntime/Public/AnimSequencerInstanceProxy.h
+++ b/Engine/Source/Runtime/AnimGraphRuntime/Public/AnimSequencerInstanceProxy.h
@@ -138,7 +138,7 @@ protected:
 	TMap<uint32, FSequencerPlayerBase*> SequencerToPlayerMap;
 
 	/** mapping from sequencer index to internal mirror node index */
-	TMap<uint32, FAnimNode_Mirror*> SequencerToMirrorMap;
+	TMap<uint32, FAnimNode_Mirror_Standalone*> SequencerToMirrorMap;
 
 	/** custom root motion override sent in from sequencer */
 	TOptional<FRootMotionOverride> RootMotionOverride;
