Blend lists now only re-initialize children when weight is zero

This prevents re-initialization pops if a non-active child still has weight, and is consistent with state machines

#7249

diff --git a/Engine/Source/Runtime/AnimGraphRuntime/Private/AnimNodes/AnimNode_BlendListBase.cpp b/Engine/Source/Runtime/AnimGraphRuntime/Private/AnimNodes/AnimNode_BlendListBase.cpp
--- a/Engine/Source/Runtime/AnimGraphRuntime/Private/AnimNodes/AnimNode_BlendListBase.cpp
+++ b/Engine/Source/Runtime/AnimGraphRuntime/Private/AnimNodes/AnimNode_BlendListBase.cpp
@@ -185,8 +185,8 @@ void FAnimNode_BlendListBase::Update_AnyThread(const FAnimationUpdateContext& Co
 				}
 			}
 
-			// when this flag is true, we'll reinitialize the children
-			if (GetResetChildOnActivation())
+			// When bResetChildOnActivation is true and the weight of the new child is zero, we'll reinitialize the child.
+			if (GetResetChildOnActivation() && CurrentWeight <= ZERO_ANIMWEIGHT_THRESH)
 			{
 				FAnimationInitializeContext ReinitializeContext(Context.AnimInstanceProxy, Context.SharedContext);
 
diff --git a/Engine/Source/Runtime/AnimGraphRuntime/Public/AnimNodes/AnimNode_BlendListBase.h b/Engine/Source/Runtime/AnimGraphRuntime/Public/AnimNodes/AnimNode_BlendListBase.h
--- a/Engine/Source/Runtime/AnimGraphRuntime/Public/AnimNodes/AnimNode_BlendListBase.h
+++ b/Engine/Source/Runtime/AnimGraphRuntime/Public/AnimNodes/AnimNode_BlendListBase.h
@@ -41,7 +41,7 @@ private:
 	EAlphaBlendOption BlendType = EAlphaBlendOption::Linear;
 	
 protected:
-	/** This reinitializes child pose when re-activated. For example, when active child changes */
+	/** This reinitializes the re-activated child if the child's weight was zero. */
 	UPROPERTY(EditAnywhere, Category = Option, meta=(FoldProperty))
 	bool bResetChildOnActivation = false;
 
