Added Animation State Machine Library for common helper functions used in  anim node functions:
*Added IsStateBlendingIn/Out to library

diff --git a/Engine/Source/Editor/AnimGraph/Private/AnimGraphNode_StateMachineBase.cpp b/Engine/Source/Editor/AnimGraph/Private/AnimGraphNode_StateMachineBase.cpp
--- a/Engine/Source/Editor/AnimGraph/Private/AnimGraphNode_StateMachineBase.cpp
+++ b/Engine/Source/Editor/AnimGraph/Private/AnimGraphNode_StateMachineBase.cpp
@@ -454,6 +454,8 @@ void UAnimGraphNode_StateMachineBase::OnProcessDuringCompilation(IAnimBlueprintC
 				{
 					InCompilationContext.ValidateGraphIsWellFormed(StateNode->BoundGraph);
 
+					AnimGraphResultNode->Node.SetStateIndex(StateIndex);
+
 					BakedState.StateRootNodeIndex = Extension->ExpandGraphAndProcessNodes(StateNode->BoundGraph, AnimGraphResultNode, InCompilationContext, OutCompiledData);
 
 					// See if the state consists of a single sequence player node, and remember the index if so
diff --git a/Engine/Source/Runtime/AnimGraphRuntime/Private/AnimationStateMachineLibrary.cpp b/Engine/Source/Runtime/AnimGraphRuntime/Private/AnimationStateMachineLibrary.cpp
new file mode 100644
--- /dev/null
+++ b/Engine/Source/Runtime/AnimGraphRuntime/Private/AnimationStateMachineLibrary.cpp
@@ -0,0 +1,80 @@
+// Copyright Epic Games, Inc. All Rights Reserved.
+
+#include "AnimationStateMachineLibrary.h"
+
+#include "Animation/AnimNode_StateMachine.h"
+#include "AnimNodes/AnimNode_StateResult.h"
+#include "Animation/AnimInstanceProxy.h"
+
+DEFINE_LOG_CATEGORY_STATIC(LogAnimationStateMachineLibrary, Verbose, All);
+
+void UAnimationStateMachineLibrary::ConvertToAnimationStateResult(const FAnimNodeReference& Node, FAnimationStateResultReference& AnimationState, EAnimNodeReferenceConversionResult& Result)
+{
+	AnimationState = FAnimNodeReference::ConvertToType<FAnimationStateResultReference>(Node, Result);
+}
+
+bool UAnimationStateMachineLibrary::IsStateBlendingIn(const FAnimUpdateContext& UpdateContext, const FAnimationStateResultReference& Node)
+{
+	bool bResult = false;
+
+	Node.CallAnimNodeFunction<FAnimNode_StateResult>(
+		TEXT("IsStateBlendingIn"),
+		[&UpdateContext, &bResult](FAnimNode_StateResult& StateResultNode)
+		{
+			if (const FAnimationUpdateContext* AnimationUpdateContext = UpdateContext.GetContext())
+			{
+				IAnimClassInterface* AnimBlueprintClass = AnimationUpdateContext->GetAnimClass();
+
+				// Previous node to an FAnimNode_StateResult is always its owning FAnimNode_StateMachine
+				const int32 MachineIndex =  AnimBlueprintClass->GetAnimNodeProperties().Num() - 1 - AnimationUpdateContext->GetPreviousNodeId();
+				const int32 StateIndex = StateResultNode.GetStateIndex();
+
+				const FAnimInstanceProxy* AnimInstanceProxy = AnimationUpdateContext->AnimInstanceProxy;
+				if (const FAnimNode_StateMachine* MachineInstance = AnimInstanceProxy->GetStateMachineInstance(MachineIndex))
+				{
+					const int32 CurrentStateIndex = MachineInstance->GetCurrentState();
+
+					const float StateWeight = AnimInstanceProxy->GetRecordedStateWeight(MachineInstance->StateMachineIndexInClass, StateIndex);
+					if ((StateWeight < 1.0f) && (CurrentStateIndex == StateIndex))
+					{
+						bResult = true;
+					}
+				}
+			}
+		});
+
+	return bResult;
+}
+
+bool UAnimationStateMachineLibrary::IsStateBlendingOut(const FAnimUpdateContext& UpdateContext, const FAnimationStateResultReference& Node)
+{
+	bool bResult = false;
+
+	Node.CallAnimNodeFunction<FAnimNode_StateResult>(
+		TEXT("IsStateBlendingOut"),
+		[&UpdateContext, &bResult](FAnimNode_StateResult& StateResultNode)
+		{
+			if (const FAnimationUpdateContext* AnimationUpdateContext = UpdateContext.GetContext())
+			{
+				IAnimClassInterface* AnimBlueprintClass = AnimationUpdateContext->GetAnimClass();
+
+				// Previous node to an FAnimNode_StateResult is always its owning FAnimNode_StateMachine
+				const int32 MachineIndex = AnimBlueprintClass->GetAnimNodeProperties().Num() - 1 - AnimationUpdateContext->GetPreviousNodeId();
+				const int32 StateIndex = StateResultNode.GetStateIndex();
+
+				const FAnimInstanceProxy* AnimInstanceProxy = AnimationUpdateContext->AnimInstanceProxy;
+				if (const FAnimNode_StateMachine* MachineInstance = AnimInstanceProxy->GetStateMachineInstance(MachineIndex))
+				{
+					const int32 CurrentStateIndex = MachineInstance->GetCurrentState();
+
+					const float StateWeight = AnimInstanceProxy->GetRecordedStateWeight(MachineInstance->StateMachineIndexInClass, StateIndex);
+					if ((StateWeight > 0.0f) && (CurrentStateIndex != StateIndex))
+					{
+						bResult = true;
+					}
+				}
+			}
+		});
+
+	return bResult;
+}
diff --git a/Engine/Source/Runtime/AnimGraphRuntime/Public/AnimNodes/AnimNode_StateResult.cpp b/Engine/Source/Runtime/AnimGraphRuntime/Public/AnimNodes/AnimNode_StateResult.cpp
--- a/Engine/Source/Runtime/AnimGraphRuntime/Public/AnimNodes/AnimNode_StateResult.cpp
+++ b/Engine/Source/Runtime/AnimGraphRuntime/Public/AnimNodes/AnimNode_StateResult.cpp
@@ -14,4 +14,9 @@ bool FAnimNode_StateResult::SerializeFromMismatchedTag(const FPropertyTag& Tag,
 	}
 
 	return false;
-}
\ No newline at end of file
+}
+
+int32 FAnimNode_StateResult::GetStateIndex() const
+{
+	return GET_ANIM_NODE_DATA(int32, StateIndex);
+}
diff --git a/Engine/Source/Runtime/AnimGraphRuntime/Public/AnimNodes/AnimNode_StateResult.h b/Engine/Source/Runtime/AnimGraphRuntime/Public/AnimNodes/AnimNode_StateResult.h
--- a/Engine/Source/Runtime/AnimGraphRuntime/Public/AnimNodes/AnimNode_StateResult.h
+++ b/Engine/Source/Runtime/AnimGraphRuntime/Public/AnimNodes/AnimNode_StateResult.h
@@ -16,6 +16,21 @@ struct ANIMGRAPHRUNTIME_API FAnimNode_StateResult : public FAnimNode_Root
 
 	/** Used to upgrade old FAnimNode_Roots to FAnimNode_StateResult */
 	bool SerializeFromMismatchedTag(const FPropertyTag& Tag, FStructuredArchive::FSlot Slot);
+
+#if WITH_EDITORONLY_DATA
+protected:
+
+	/** The index of the state this node belongs to. Filled in during the owning state machine's compilation. */
+	UPROPERTY(meta = (FoldProperty))
+	int32 StateIndex = -1;
+#endif
+
+public:
+#if WITH_EDITORONLY_DATA
+	void SetStateIndex(int32 InStateIndex) { StateIndex = InStateIndex; }
+#endif
+
+	int32 GetStateIndex() const;
 };
 
 template<>
diff --git a/Engine/Source/Runtime/AnimGraphRuntime/Public/AnimationStateMachineLibrary.h b/Engine/Source/Runtime/AnimGraphRuntime/Public/AnimationStateMachineLibrary.h
new file mode 100644
--- /dev/null
+++ b/Engine/Source/Runtime/AnimGraphRuntime/Public/AnimationStateMachineLibrary.h
@@ -0,0 +1,50 @@
+// Copyright Epic Games, Inc. All Rights Reserved.
+
+#pragma once
+
+#include "CoreMinimal.h"
+
+#include "Animation/AnimExecutionContext.h"
+#include "Kismet/BlueprintFunctionLibrary.h"
+#include "Animation/AnimNodeReference.h"
+#include "AnimationStateMachineLibrary.generated.h"
+
+struct FAnimNode_AnimationStateMachine;
+struct FAnimNode_StateResult;
+
+USTRUCT(BlueprintType, DisplayName = "Animation State Reference")
+struct FAnimationStateResultReference : public FAnimNodeReference
+{
+	GENERATED_BODY()
+
+	typedef FAnimNode_StateResult FInternalNodeType;
+};
+
+// Exposes operations to be performed on anim state machine node contexts
+UCLASS(Experimental)
+class ANIMGRAPHRUNTIME_API UAnimationStateMachineLibrary : public UBlueprintFunctionLibrary
+{
+	GENERATED_BODY()
+
+public:
+	/** Get an anim state reference from an anim node reference */
+	UFUNCTION(BlueprintCallable, Category = "State Machine", meta=(BlueprintThreadSafe, DisplayName = "Convert to Animation State", ExpandEnumAsExecs = "Result"))
+	static void ConvertToAnimationStateResult(const FAnimNodeReference& Node, FAnimationStateResultReference& AnimationState, EAnimNodeReferenceConversionResult& Result);
+
+	/** Get an anim state reference from an anim node reference (pure) */
+	UFUNCTION(BlueprintPure, Category = "State Machine", meta=(BlueprintThreadSafe, DisplayName = "Convert to Animation State"))
+	static void ConvertToAnimationStateResultPure(const FAnimNodeReference& Node, FAnimationStateResultReference& AnimationState, bool& Result)
+	{
+		EAnimNodeReferenceConversionResult ConversionResult;
+		 ConvertToAnimationStateResult(Node, AnimationState, ConversionResult);
+		Result = (ConversionResult == EAnimNodeReferenceConversionResult::Succeeded);	
+	}
+	
+	/** Returns whether the state the node belongs to is blending in */
+	UFUNCTION(BlueprintPure, Category = "State Machine", meta=(BlueprintThreadSafe))
+	static bool IsStateBlendingIn(const FAnimUpdateContext& UpdateContext, const FAnimationStateResultReference& Node);
+
+	/** Returns whether the state the node belongs to is blending out */
+	UFUNCTION(BlueprintPure, Category = "State Machine", meta=(BlueprintThreadSafe))
+	static bool IsStateBlendingOut(const FAnimUpdateContext& UpdateContext, const FAnimationStateResultReference& Node);
+};
\ No newline at end of file
