Enable warning for scaled RBAN objects with non-worldspace sim and world space collisions

diff --git a/Engine/Source/Runtime/AnimGraphRuntime/Private/BoneControllers/AnimNode_RigidBody.cpp b/Engine/Source/Runtime/AnimGraphRuntime/Private/BoneControllers/AnimNode_RigidBody.cpp
--- a/Engine/Source/Runtime/AnimGraphRuntime/Private/BoneControllers/AnimNode_RigidBody.cpp
+++ b/Engine/Source/Runtime/AnimGraphRuntime/Private/BoneControllers/AnimNode_RigidBody.cpp
@@ -1505,10 +1505,10 @@ void FAnimNode_RigidBody::PreUpdate(const UAnimInstance* InAnimInstance)
 	APawn* PawnOwner = InAnimInstance->TryGetPawnOwner();
 	UPawnMovementComponent* MovementComp = PawnOwner ? PawnOwner->GetMovementComponent() : nullptr;
 
-#if WITH_EDITOR && !WITH_CHAOS
-	if (bEnableWorldGeometry && SimulationSpace != ESimulationSpace::WorldSpace)
+#if WITH_EDITOR && WITH_CHAOS
+	if (bEnableWorldGeometry && SimulationSpace != ESimulationSpace::WorldSpace && SKC && SKC->GetRelativeScale3D() != FVector(1.f, 1.f, 1.f))
 	{
-		FMessageLog("PIE").Warning(FText::Format(LOCTEXT("WorldCollisionComponentSpace", "Trying to use world collision without world space simulation for ''{0}''. This is not supported, please change SimulationSpace to WorldSpace"),
+		FMessageLog("PIE").Warning(FText::Format(LOCTEXT("WorldCollisionComponentSpace", "Trying to use world collision without world space simulation for scaled ''{0}''. This is not supported, please change SimulationSpace to WorldSpace"),
 			FText::FromString(GetPathNameSafe(SKC))));
 	}
 #endif
