Consolidate SoA into AoS in FAnimNode_BlendListBase

Savings ~1.4 MB peak in synthetic tests

diff --git a/Engine/Source/Runtime/AnimGraphRuntime/Private/AnimNodes/AnimNode_BlendListBase.cpp b/Engine/Source/Runtime/AnimGraphRuntime/Private/AnimNodes/AnimNode_BlendListBase.cpp
--- a/Engine/Source/Runtime/AnimGraphRuntime/Private/AnimNodes/AnimNode_BlendListBase.cpp
+++ b/Engine/Source/Runtime/AnimGraphRuntime/Private/AnimNodes/AnimNode_BlendListBase.cpp
@@ -19,16 +19,13 @@ void FAnimNode_BlendListBase::Initialize_AnyThread(const FAnimationInitializeCon
 	const TArray<float>& CurrentBlendTimes = GetBlendTimes();
 	checkSlow(CurrentBlendTimes.Num() == NumPoses);
 
-	BlendWeights.Reset(NumPoses);
-	PosesToEvaluate.Reset(NumPoses);
+	PerBlendData.Reset(NumPoses);
 	if (NumPoses > 0)
 	{
 		// If we have at least 1 pose we initialize to full weight on
 		// the first pose
-		BlendWeights.AddZeroed(NumPoses);
-		BlendWeights[0] = 1.0f;
-
-		PosesToEvaluate.Add(0);
+		PerBlendData.AddZeroed(NumPoses);
+		PerBlendData[0].Weight = 1.0f;
 
 		for (int32 ChildIndex = 0; ChildIndex < NumPoses; ++ChildIndex)
 		{
@@ -36,25 +33,15 @@ void FAnimNode_BlendListBase::Initialize_AnyThread(const FAnimationInitializeCon
 		}
 	}
 
-	RemainingBlendTimes.Empty(NumPoses);
-	RemainingBlendTimes.AddZeroed(NumPoses);
-	Blends.Empty(NumPoses);
-	Blends.AddZeroed(NumPoses);
-
 	UBlendProfile* CurrentBlendProfile = GetBlendProfile();
-	if (CurrentBlendProfile)
-	{
-		BlendStartAlphas.Empty(NumPoses);
-		BlendStartAlphas.AddZeroed(NumPoses);
-	}
 
 	LastActiveChildIndex = INDEX_NONE;
 
 	EAlphaBlendOption CurrentBlendType = GetBlendType();
 	UCurveFloat* CurrentCustomBlendCurve = GetCustomBlendCurve();
-	for(int32 i = 0 ; i < Blends.Num() ; ++i)
+	for(int32 i = 0 ; i < PerBlendData.Num() ; ++i)
 	{
-		FAlphaBlend& Blend = Blends[i];
+		FAlphaBlend& Blend = PerBlendData[i].Blend;
 
 		Blend.SetBlendTime(0.0f);
 		Blend.SetBlendOption(CurrentBlendType);
@@ -62,14 +49,14 @@ void FAnimNode_BlendListBase::Initialize_AnyThread(const FAnimationInitializeCon
 
 		if (CurrentBlendProfile)
 		{
-			BlendStartAlphas[i] = 0.0f;
+			PerBlendData[i].StartAlpha = 0.0f;
 		}
 	}
-	Blends[0].SetAlpha(1.0f);
+	PerBlendData[0].Blend.SetAlpha(1.0f);
 
 	if(CurrentBlendProfile)
 	{
-		BlendStartAlphas[0] = 1.0f;
+		PerBlendData[0].StartAlpha = 1.0f;
 
 		// Initialise per-bone data
 		PerBoneSampleData.Empty(NumPoses);
@@ -100,9 +87,7 @@ void FAnimNode_BlendListBase::Update_AnyThread(const FAnimationUpdateContext& Co
 
 	const int32 NumPoses = BlendPose.Num();
 	const TArray<float>& CurrentBlendTimes = GetBlendTimes();
-	checkSlow((CurrentBlendTimes.Num() == NumPoses) && (BlendWeights.Num() == NumPoses));
-
-	PosesToEvaluate.Empty(NumPoses);
+	checkSlow(PerBlendData.Num() == NumPoses);
 
 	if (NumPoses > 0)
 	{
@@ -115,7 +100,7 @@ void FAnimNode_BlendListBase::Update_AnyThread(const FAnimationUpdateContext& Co
 		{
 			bool LastChildIndexIsInvalid = (LastActiveChildIndex == INDEX_NONE);
 			
-			const float CurrentWeight = BlendWeights[ChildIndex];
+			const float CurrentWeight = PerBlendData[ChildIndex].Weight;
 			const float DesiredWeight = 1.0f;
 			const float WeightDifference = FMath::Clamp<float>(FMath::Abs<float>(DesiredWeight - CurrentWeight), 0.0f, 1.0f);
 
@@ -147,9 +132,9 @@ void FAnimNode_BlendListBase::Update_AnyThread(const FAnimationUpdateContext& Co
 				RemainingBlendTime = CurrentBlendTimes[ChildIndex] * WeightDifference;
 			}
 
-			for (int32 i = 0; i < RemainingBlendTimes.Num(); ++i)
+			for (int32 i = 0; i < PerBlendData.Num(); ++i)
 			{
-				RemainingBlendTimes[i] = RemainingBlendTime;
+				PerBlendData[i].RemainingTime = RemainingBlendTime;
 			}
 
 			// If we have a valid previous child and we're instantly blending - update that pose with zero weight
@@ -158,30 +143,30 @@ void FAnimNode_BlendListBase::Update_AnyThread(const FAnimationUpdateContext& Co
 				BlendPose[LastActiveChildIndex].Update(Context.FractionalWeight(0.0f));
 			}
 
-			for(int32 i = 0; i < Blends.Num(); ++i)
+			for(int32 i = 0; i < PerBlendData.Num(); ++i)
 			{
-				FAlphaBlend& Blend = Blends[i];
+				FAlphaBlend& Blend = PerBlendData[i].Blend;
 
 				Blend.SetBlendTime(RemainingBlendTime);
 
 				if(i == ChildIndex)
 				{
-					Blend.SetValueRange(BlendWeights[i], 1.0f);
+					Blend.SetValueRange(PerBlendData[i].Weight, 1.0f);
 
 					if (CurrentBlendProfile)
 					{
 						Blend.ResetAlpha();
-						BlendStartAlphas[i] = Blend.GetAlpha();
+						PerBlendData[i].StartAlpha = Blend.GetAlpha();
 					}
 				}
 				else
 				{
-					Blend.SetValueRange(BlendWeights[i], 0.0f);
+					Blend.SetValueRange(PerBlendData[i].Weight, 0.0f);
 				}
 
 				if (CurrentBlendProfile)
 				{
-					BlendStartAlphas[i] = Blend.GetAlpha();
+					PerBlendData[i].StartAlpha = Blend.GetAlpha();
 				}
 			}
 
@@ -200,11 +185,11 @@ void FAnimNode_BlendListBase::Update_AnyThread(const FAnimationUpdateContext& Co
 		// Advance the weights
 		//@TODO: This means we advance even in a frame where the target weights/times just got modified; is that desirable?
 		float SumWeight = 0.0f;
-		for (int32 i = 0; i < Blends.Num(); ++i)
+		for (int32 i = 0; i < PerBlendData.Num(); ++i)
 		{
-			float& BlendWeight = BlendWeights[i];
+			float& BlendWeight = PerBlendData[i].Weight;
 
-			FAlphaBlend& Blend = Blends[i];
+			FAlphaBlend& Blend = PerBlendData[i].Blend;
 			Blend.Update(Context.GetDeltaTime());
 			BlendWeight = Blend.GetBlendedValue();
 
@@ -215,20 +200,19 @@ void FAnimNode_BlendListBase::Update_AnyThread(const FAnimationUpdateContext& Co
 		if ((SumWeight > ZERO_ANIMWEIGHT_THRESH) && (FMath::Abs<float>(SumWeight - 1.0f) > ZERO_ANIMWEIGHT_THRESH))
 		{
 			float ReciprocalSum = 1.0f / SumWeight;
-			for (int32 i = 0; i < BlendWeights.Num(); ++i)
+			for (int32 i = 0; i < PerBlendData.Num(); ++i)
 			{
-				BlendWeights[i] *= ReciprocalSum;
+				PerBlendData[i].Weight *= ReciprocalSum;
 			}
 		}
 
 		// Update our active children
 		for (int32 i = 0; i < BlendPose.Num(); ++i)
 		{
-			const float BlendWeight = BlendWeights[i];
+			const float BlendWeight = PerBlendData[i].Weight;
 			if (BlendWeight > ZERO_ANIMWEIGHT_THRESH)
 			{
 				BlendPose[i].Update(Context.FractionalWeight(BlendWeight));
-				PosesToEvaluate.Add(i);				
 			}
 		}
 
@@ -238,9 +222,10 @@ void FAnimNode_BlendListBase::Update_AnyThread(const FAnimationUpdateContext& Co
 			for(int32 i = 0; i < BlendPose.Num(); ++i)
 			{
 				FBlendSampleData& PoseSampleData = PerBoneSampleData[i];
-				PoseSampleData.TotalWeight = BlendWeights[i];
+				const FBlendData& BlendData = PerBlendData[i];
+				PoseSampleData.TotalWeight = BlendData.Weight;
 				const bool bInverse = (CurrentBlendProfile->Mode == EBlendProfileMode::WeightFactor) ? (ChildIndex != i) : false;
-				CurrentBlendProfile->UpdateBoneWeights(PoseSampleData, Blends[i], BlendStartAlphas[i], BlendWeights[i], bInverse);
+				CurrentBlendProfile->UpdateBoneWeights(PoseSampleData, BlendData.Blend, BlendData.StartAlpha, BlendData.Weight, bInverse);
 			}
 
 			FBlendSampleData::NormalizeDataWeight(PerBoneSampleData);
@@ -250,7 +235,7 @@ void FAnimNode_BlendListBase::Update_AnyThread(const FAnimationUpdateContext& Co
 #if ANIM_TRACE_ENABLED
 	const int32 ChildIndex = GetActiveChildIndex();
 	TRACE_ANIM_NODE_VALUE(Context, TEXT("Active Index"), ChildIndex);
-	TRACE_ANIM_NODE_VALUE(Context, TEXT("Active Weight"), BlendWeights[ChildIndex]);
+	TRACE_ANIM_NODE_VALUE(Context, TEXT("Active Weight"), PerBlendData[ChildIndex].Weight);
 	TRACE_ANIM_NODE_VALUE(Context, TEXT("Active Blend Time"), CurrentBlendTimes[ChildIndex]);
 #endif
 }
@@ -260,6 +245,20 @@ void FAnimNode_BlendListBase::Evaluate_AnyThread(FPoseContext& Output)
 	DECLARE_SCOPE_HIERARCHICAL_COUNTER_ANIMNODE(Evaluate_AnyThread)
 	ANIM_MT_SCOPE_CYCLE_COUNTER(BlendPosesInGraph, !IsInGameThread());
 
+	// Build local arrays to pass to the BlendPosesTogether runtime
+	TArray<int32, TInlineAllocator<8>> PosesToEvaluate;
+	TArray<float, TInlineAllocator<8>> BlendWeights;
+	const int32 PerBlendDataCount = PerBlendData.Num();
+	for(int32 BlendIndex = 0; BlendIndex < PerBlendDataCount; ++BlendIndex)
+	{
+		const FBlendData& BlendData = PerBlendData[BlendIndex];
+		BlendWeights.Add(BlendData.Weight);
+		if(BlendData.Weight > ZERO_ANIMWEIGHT_THRESH)
+		{
+			PosesToEvaluate.Add(BlendIndex);
+		}
+	}
+
 	const int32 NumPoses = PosesToEvaluate.Num();
 
 	if ((NumPoses > 0) && (BlendPose.Num() == BlendWeights.Num()))
@@ -315,13 +314,13 @@ void FAnimNode_BlendListBase::GatherDebugData(FNodeDebugData& DebugData)
 
 	FString DebugLine = GetNodeName(DebugData);
 	const TArray<float>& CurrentBlendTimes = GetBlendTimes();
-	DebugLine += FString::Printf(TEXT("(Active: (%i/%i) Weight: %.1f%% Time %.3f)"), ChildIndex+1, NumPoses, BlendWeights[ChildIndex]*100.f, CurrentBlendTimes[ChildIndex]);
+	DebugLine += FString::Printf(TEXT("(Active: (%i/%i) Weight: %.1f%% Time %.3f)"), ChildIndex+1, NumPoses, PerBlendData[ChildIndex].Weight*100.f, CurrentBlendTimes[ChildIndex]);
 
 	DebugData.AddDebugItem(DebugLine);
 	
 	for(int32 Pose = 0; Pose < NumPoses; ++Pose)
 	{
-		BlendPose[Pose].GatherDebugData(DebugData.BranchFlow(BlendWeights[Pose]));
+		BlendPose[Pose].GatherDebugData(DebugData.BranchFlow(PerBlendData[Pose].Weight));
 	}
 }
 
diff --git a/Engine/Source/Runtime/AnimGraphRuntime/Public/AnimNodes/AnimNode_BlendListBase.h b/Engine/Source/Runtime/AnimGraphRuntime/Public/AnimNodes/AnimNode_BlendListBase.h
--- a/Engine/Source/Runtime/AnimGraphRuntime/Public/AnimNodes/AnimNode_BlendListBase.h
+++ b/Engine/Source/Runtime/AnimGraphRuntime/Public/AnimNodes/AnimNode_BlendListBase.h
@@ -52,21 +52,22 @@ private:
 	UPROPERTY(EditAnywhere, Category=BlendType, meta=(UseAsBlendProfile=true, FoldProperty))
 	TObjectPtr<UBlendProfile> BlendProfile = nullptr;
 #endif // #if WITH_EDITORONLY_DATA
-	
-	TArray<struct FAlphaBlend> Blends;
 
 protected:
-	TArray<float> BlendWeights;
-
-	TArray<float> RemainingBlendTimes;
+	// Struct for tracking blends for each pose
+	struct FBlendData
+	{
+		FAlphaBlend Blend;
+		float Weight;
+		float RemainingTime;
+		float StartAlpha;
+	};
 
-	TArray<float> BlendStartAlphas;
+	TArray<FBlendData> PerBlendData;
 
+	// Per-bone blending data, allocated when using blend profiles
 	TArray<FBlendSampleData> PerBoneSampleData;
 
-	//Store which poses we need to evaluate
-	TArray<int32> PosesToEvaluate;
-
 	int32 LastActiveChildIndex = 0;
 	
 public:	
