Fixed a long stall that could occur with the deferred RBAN task. Migrated the task to the new UE5 task interface that will retract it from the task queue and run it inline if it doesn't get started by the time it's flushed next frame.

diff --git a/Engine/Source/Runtime/AnimGraphRuntime/Private/BoneControllers/AnimNode_RigidBody.cpp b/Engine/Source/Runtime/AnimGraphRuntime/Private/BoneControllers/AnimNode_RigidBody.cpp
--- a/Engine/Source/Runtime/AnimGraphRuntime/Private/BoneControllers/AnimNode_RigidBody.cpp
+++ b/Engine/Source/Runtime/AnimGraphRuntime/Private/BoneControllers/AnimNode_RigidBody.cpp
@@ -77,50 +77,20 @@ FAutoConsoleVariableRef CVarRigidBodyNodeDeferredSimulation(TEXT("p.RigidBodyNod
 FAutoConsoleVariableRef CVarRigidBodyNodeDebugDraw(TEXT("p.RigidBodyNode.DebugDraw"), bRBAN_DebugDraw, TEXT("Whether to debug draw the rigid body simulation state. Requires p.Chaos.DebugDraw.Enabled 1 to function as well."), ECVF_Default);
 
 // Array of priorities that can be indexed into with CVars, since task priorities cannot be set from scalability .ini
-static FAutoConsoleTaskPriority GRigidBodyNodeTaskPriorities[] =
+static UE::Tasks::ETaskPriority GRigidBodyNodeTaskPriorities[] =
 {
-	FAutoConsoleTaskPriority(
-		TEXT("p.RigidBodyNode.TaskPriorities.Maximum"),
-		TEXT("Maximum priority for rigid body node simulation"),
-		ENamedThreads::HighThreadPriority,
-		ENamedThreads::HighTaskPriority,
-		ENamedThreads::HighTaskPriority // High task priority when forced to normal thread priority
-	),
-	FAutoConsoleTaskPriority(
-		TEXT("p.RigidBodyNode.TaskPriorities.High"),
-		TEXT("High priority for rigid body node simulation"),
-		ENamedThreads::HighThreadPriority,
-		ENamedThreads::NormalTaskPriority,
-		ENamedThreads::NormalTaskPriority // Normal (lowest) task priority if forced to normal thread priority
-	),
-	FAutoConsoleTaskPriority(
-		TEXT("p.RigidBodyNode.TaskPriorities.Normal"),
-		TEXT("Normal priority for rigid body node simulation"),
-		ENamedThreads::NormalThreadPriority,
-		ENamedThreads::NormalTaskPriority,
-		ENamedThreads::NormalTaskPriority // Normal (lowest) task priority if forced to normal thread priority
-	),
-	FAutoConsoleTaskPriority(
-		TEXT("p.RigidBodyNode.TaskPriorities.Low"),
-		TEXT("Low priority for rigid body node simulation"),
-		ENamedThreads::BackgroundThreadPriority,
-		ENamedThreads::HighTaskPriority,
-		ENamedThreads::NormalTaskPriority // Normal (lowest) task priority if forced to normal thread priority
-	),
-	FAutoConsoleTaskPriority(
-		TEXT("p.RigidBodyNode.TaskPriorities.Minimum"),
-		TEXT("Minimum priority for rigid body node simulation"),
-		ENamedThreads::BackgroundThreadPriority,
-		ENamedThreads::NormalTaskPriority,
-		ENamedThreads::NormalTaskPriority // Normal (lowest) task priority if forced to normal thread priority
-	),
+	UE::Tasks::ETaskPriority::High,
+	UE::Tasks::ETaskPriority::Normal,
+	UE::Tasks::ETaskPriority::BackgroundHigh,
+	UE::Tasks::ETaskPriority::BackgroundNormal,
+	UE::Tasks::ETaskPriority::BackgroundLow
 };
 
 static int32 GRigidBodyNodeSimulationTaskPriority = 0;
 FAutoConsoleVariableRef CVarRigidBodyNodeSimulationTaskPriority(
 	TEXT("p.RigidBodyNode.TaskPriority.Simulation"),
 	GRigidBodyNodeSimulationTaskPriority,
-	TEXT("Task priority index for running the rigid body node simulation task"),
+	TEXT("Task priority for running the rigid body node simulation task (0 = foreground/high, 1 = foreground/normal, 2 = background/high, 3 = background/normal, 4 = background/low)."),
 	ECVF_Default
 );
 
@@ -189,7 +159,7 @@ FAnimNode_RigidBody::FAnimNode_RigidBody()
 	, SkelMeshCompWeakPtr()
 	, PhysicsSimulation(nullptr)
 	, SolverIterations()
-	, SimulationTaskState()
+	, SimulationTask()
 	, OutputBoneData()
 	, Bodies()
 	, SkeletonBoneIndexToBodyIndex()
@@ -504,43 +474,6 @@ DECLARE_CYCLE_STAT(TEXT("RigidBodyNode_Simulation"), STAT_RigidBodyNode_Simulati
 DECLARE_CYCLE_STAT(TEXT("RigidBodyNode_SimulationWait"), STAT_RigidBodyNode_SimulationWait, STATGROUP_Anim);
 DECLARE_CYCLE_STAT(TEXT("FAnimNode_RigidBody::EvaluateSkeletalControl_AnyThread"), STAT_ImmediateEvaluateSkeletalControl, STATGROUP_ImmediatePhysics);
 
-class FRigidBodyNodeSimulationTask
-{
-public:
-	FRigidBodyNodeSimulationTask(FAnimNode_RigidBody* InRigidBodyNode, const float InDeltaSeconds, const FVector& InSimSpaceGravity)
-		: RigidBodyNode(InRigidBodyNode)
-		, DeltaSeconds(InDeltaSeconds)
-		, SimSpaceGravity(InSimSpaceGravity)
-	{
-	}
-
-	static FORCEINLINE TStatId GetStatId()
-	{
-		RETURN_QUICK_DECLARE_CYCLE_STAT(FRigidBodyNodeSimulationTask, STATGROUP_TaskGraphTasks);
-	}
-
-	static FORCEINLINE ENamedThreads::Type GetDesiredThread()
-	{
-		const int32 PriorityIndex = FMath::Clamp<int32>(GRigidBodyNodeSimulationTaskPriority, 0, UE_ARRAY_COUNT(GRigidBodyNodeTaskPriorities) - 1);
-		return GRigidBodyNodeTaskPriorities[PriorityIndex].Get();
-	}
-
-	static FORCEINLINE ESubsequentsMode::Type GetSubsequentsMode()
-	{
-		return ESubsequentsMode::TrackSubsequents;
-	}
-
-	void DoTask(ENamedThreads::Type CurrentThread, const FGraphEventRef& CompletionGraphEvent)
-	{
-		RigidBodyNode->RunPhysicsSimulation(DeltaSeconds, SimSpaceGravity);
-	}
-
-private:
-	FAnimNode_RigidBody* const RigidBodyNode;
-	const float DeltaSeconds;
-	const FVector SimSpaceGravity;
-};
-
 void FAnimNode_RigidBody::RunPhysicsSimulation(float DeltaSeconds, const FVector& SimSpaceGravity)
 {
 	SCOPE_CYCLE_COUNTER(STAT_RigidBodyNode_Simulation);
@@ -565,12 +498,11 @@ void FAnimNode_RigidBody::RunPhysicsSimulation(float DeltaSeconds, const FVector
 
 void FAnimNode_RigidBody::FlushDeferredSimulationTask()
 {
-	FGraphEventRef& CompletionEvent = SimulationTaskState.SimulationCompletionEvent;
-	if (CompletionEvent && !CompletionEvent->IsComplete())
+	if (SimulationTask.IsValid() && !SimulationTask.IsCompleted())
 	{
 		SCOPE_CYCLE_COUNTER(STAT_RigidBodyNode_SimulationWait);
 		CSV_SCOPED_TIMING_STAT(Animation, RigidBodyNodeSimulationWait);
-		FTaskGraphInterface::Get().WaitUntilTasksComplete({CompletionEvent});
+		SimulationTask.Wait();
 	}
 }
 
@@ -986,8 +918,13 @@ void FAnimNode_RigidBody::EvaluateSkeletalControl_AnyThread(FComponentSpacePoseC
 		if (bNeedsSimulationTick && bUseDeferredSimulationTask)
 		{
 			// FlushDeferredSimulationTask() should have already ensured task is done.
-			ensure(!SimulationTaskState.SimulationCompletionEvent || SimulationTaskState.SimulationCompletionEvent->IsComplete());
-			SimulationTaskState.SimulationCompletionEvent = TGraphTask<FRigidBodyNodeSimulationTask>::CreateTask().ConstructAndDispatchWhenReady(this, DeltaSeconds, SimSpaceGravity);
+			ensure(SimulationTask.IsCompleted());
+			const int32 PriorityIndex = FMath::Clamp<int32>(GRigidBodyNodeSimulationTaskPriority, 0, UE_ARRAY_COUNT(GRigidBodyNodeTaskPriorities) - 1);
+			const UE::Tasks::ETaskPriority TaskPriority = GRigidBodyNodeTaskPriorities[PriorityIndex];
+			SimulationTask = UE::Tasks::Launch(
+				TEXT("RigidBodyNodeSimulationTask"),
+				[this, DeltaSeconds, SimSpaceGravity] { RunPhysicsSimulation(DeltaSeconds, SimSpaceGravity); },
+				TaskPriority);
 		}
 
 		PreviousCompWorldSpaceTM = CompWorldSpaceTM;
diff --git a/Engine/Source/Runtime/AnimGraphRuntime/Public/BoneControllers/AnimNode_RigidBody.h b/Engine/Source/Runtime/AnimGraphRuntime/Public/BoneControllers/AnimNode_RigidBody.h
--- a/Engine/Source/Runtime/AnimGraphRuntime/Public/BoneControllers/AnimNode_RigidBody.h
+++ b/Engine/Source/Runtime/AnimGraphRuntime/Public/BoneControllers/AnimNode_RigidBody.h
@@ -342,23 +342,7 @@ private:
 	FSolverIterations SolverIterations;
 
 	friend class FRigidBodyNodeSimulationTask;
-	struct FSimulationTaskState
-	{
-		FSimulationTaskState()
-			: SimulationCompletionEvent(nullptr)
-		{
-		}
-
-		// We must explicitly declare the assignment operator due to the atomic, but we must never assign to or from a state with a live simulation
-		FORCEINLINE FSimulationTaskState& operator=(const FSimulationTaskState& Other)
-		{
-			check((SimulationCompletionEvent == nullptr) && (Other.SimulationCompletionEvent == nullptr));
-			return(*this);
-		}
-
-		FGraphEventRef SimulationCompletionEvent;
-	};
-	FSimulationTaskState SimulationTaskState;
+	UE::Tasks::FTask SimulationTask;
 
 	struct FOutputBoneData
 	{
