Added anim node function getter for accumulated time/current sequence for sequence evaluators

diff --git a/Engine/Source/Runtime/AnimGraphRuntime/Private/SequenceEvaluatorLibrary.cpp b/Engine/Source/Runtime/AnimGraphRuntime/Private/SequenceEvaluatorLibrary.cpp
--- a/Engine/Source/Runtime/AnimGraphRuntime/Private/SequenceEvaluatorLibrary.cpp
+++ b/Engine/Source/Runtime/AnimGraphRuntime/Private/SequenceEvaluatorLibrary.cpp
@@ -99,3 +99,29 @@ FSequenceEvaluatorReference USequenceEvaluatorLibrary::SetSequenceWithInertialBl
 
 	return SequenceEvaluator;
 }
+
+float USequenceEvaluatorLibrary::GetAccumulatedTime(const FSequenceEvaluatorReference& SequenceEvaluator)
+{
+	float OutAccumulatedTime = 0.0f;
+	SequenceEvaluator.CallAnimNodeFunction<FAnimNode_SequenceEvaluator>(
+		TEXT("GetAccumulatedTime"),
+		[&OutAccumulatedTime](FAnimNode_SequenceEvaluator& InSequenceEvaluator)
+		{
+			OutAccumulatedTime = InSequenceEvaluator.GetAccumulatedTime();
+		});
+
+	return OutAccumulatedTime;
+}
+
+UAnimSequenceBase* USequenceEvaluatorLibrary::GetSequence(const FSequenceEvaluatorReference& SequenceEvaluator)
+{
+	UAnimSequenceBase* OutSequence = nullptr;
+	SequenceEvaluator.CallAnimNodeFunction<FAnimNode_SequenceEvaluator>(
+		TEXT("GetAccumulatedTime"),
+		[&OutSequence](FAnimNode_SequenceEvaluator& InSequenceEvaluator)
+		{
+			OutSequence = InSequenceEvaluator.GetSequence();
+		});
+
+	return OutSequence;
+}
\ No newline at end of file
diff --git a/Engine/Source/Runtime/AnimGraphRuntime/Public/SequenceEvaluatorLibrary.h b/Engine/Source/Runtime/AnimGraphRuntime/Public/SequenceEvaluatorLibrary.h
--- a/Engine/Source/Runtime/AnimGraphRuntime/Public/SequenceEvaluatorLibrary.h
+++ b/Engine/Source/Runtime/AnimGraphRuntime/Public/SequenceEvaluatorLibrary.h
@@ -56,4 +56,11 @@ public:
 	UFUNCTION(BlueprintCallable, Category = "Sequence Evaluator", meta=(BlueprintThreadSafe))
 	static FSequenceEvaluatorReference SetSequenceWithInertialBlending(const FAnimUpdateContext& UpdateContext, const FSequenceEvaluatorReference& SequenceEvaluator, UAnimSequenceBase* Sequence, float BlendTime = 0.2f);
 	
+	/** Gets the current accumulated time of the sequence evaluator */
+	UFUNCTION(BlueprintPure, Category = "Sequence Evaluator", meta = (BlueprintThreadSafe))
+	static float GetAccumulatedTime(const FSequenceEvaluatorReference& SequenceEvaluator);
+
+	/** Gets the current sequence of the sequence evaluator */
+	UFUNCTION(BlueprintPure, Category = "Sequence Evaluator", meta = (BlueprintThreadSafe))
+	static UAnimSequenceBase* GetSequence(const FSequenceEvaluatorReference& SequenceEvaluator);
 };
\ No newline at end of file
