Lazy initialization of AnimNode_RigidBody - Memory saving optimization.

diff --git a/Engine/Source/Runtime/AnimGraphRuntime/Private/BoneControllers/AnimNode_RigidBody.cpp b/Engine/Source/Runtime/AnimGraphRuntime/Private/BoneControllers/AnimNode_RigidBody.cpp
--- a/Engine/Source/Runtime/AnimGraphRuntime/Private/BoneControllers/AnimNode_RigidBody.cpp
+++ b/Engine/Source/Runtime/AnimGraphRuntime/Private/BoneControllers/AnimNode_RigidBody.cpp
@@ -148,6 +148,7 @@ FAnimNode_RigidBody::FAnimNode_RigidBody()
 	, bEnabled(false)
 	, bSimulationStarted(false)
 	, bCheckForBodyTransformInit(false)
+	, bInitialized(false)
 #if WITH_EDITORONLY_DATA
 	, bComponentSpaceSimulation_DEPRECATED(true)
 #endif
@@ -1430,14 +1431,21 @@ DECLARE_CYCLE_STAT(TEXT("RigidBody_PreUpdate"), STAT_RigidBody_PreUpdate, STATGR
 
 void FAnimNode_RigidBody::PreUpdate(const UAnimInstance* InAnimInstance)
 {
+	SCOPE_CYCLE_COUNTER(STAT_RigidBody_PreUpdate);
+	LLM_SCOPE_BYNAME(TEXT("Animation/RigidBody"));
+
+	if (!bInitialized)
+	{
+		InitPhysics(InAnimInstance);
+		bInitialized = true;
+	}
+
 	// Don't update geometry if RBN is disabled
 	if(!bEnabled)
 	{
 		return;
 	}
 
-	SCOPE_CYCLE_COUNTER(STAT_RigidBody_PreUpdate);
-
 	USkeletalMeshComponent* SKC = InAnimInstance->GetSkelMeshComponent();
 	APawn* PawnOwner = InAnimInstance->TryGetPawnOwner();
 	UPawnMovementComponent* MovementComp = PawnOwner ? PawnOwner->GetMovementComponent() : nullptr;
@@ -1504,7 +1512,8 @@ DECLARE_CYCLE_STAT(TEXT("RigidBody_Update"), STAT_RigidBody_Update, STATGROUP_An
 
 void FAnimNode_RigidBody::UpdateInternal(const FAnimationUpdateContext& Context)
 {
-	LLM_SCOPE_BYNAME(TEXT("Animation/RigidBody")); 
+	LLM_SCOPE_BYNAME(TEXT("Animation/RigidBody"));
+	
 	DECLARE_SCOPE_HIERARCHICAL_COUNTER_ANIMNODE(UpdateInternal)
 	// Avoid this work if RBN is disabled, as the results would be discarded
 	if(!bEnabled)
@@ -1801,11 +1810,6 @@ void FAnimNode_RigidBody::AddImpulseAtLocation(FVector Impulse, FVector Location
 #endif
 }
 
-void FAnimNode_RigidBody::OnInitializeAnimInstance(const FAnimInstanceProxy* InProxy, const UAnimInstance* InAnimInstance)
-{
-	InitPhysics(InAnimInstance);
-}
-
 #if WITH_EDITORONLY_DATA
 void FAnimNode_RigidBody::PostSerialize(const FArchive& Ar)
 {
diff --git a/Engine/Source/Runtime/AnimGraphRuntime/Public/BoneControllers/AnimNode_RigidBody.h b/Engine/Source/Runtime/AnimGraphRuntime/Public/BoneControllers/AnimNode_RigidBody.h
--- a/Engine/Source/Runtime/AnimGraphRuntime/Public/BoneControllers/AnimNode_RigidBody.h
+++ b/Engine/Source/Runtime/AnimGraphRuntime/Public/BoneControllers/AnimNode_RigidBody.h
@@ -132,7 +132,6 @@ struct ANIMGRAPHRUNTIME_API FAnimNode_RigidBody : public FAnimNode_SkeletalContr
 	virtual void UpdateComponentPose_AnyThread(const FAnimationUpdateContext& Context) override;
 	virtual void EvaluateComponentPose_AnyThread(FComponentSpacePoseContext& Output) override;
 	virtual void EvaluateSkeletalControl_AnyThread(FComponentSpacePoseContext& Output, TArray<FBoneTransform>& OutBoneTransforms) override;
-	virtual void OnInitializeAnimInstance(const FAnimInstanceProxy* InProxy, const UAnimInstance* InAnimInstance) override;
 	virtual bool NeedsOnInitializeAnimInstance() const override { return true; }
 	virtual void PreUpdate(const UAnimInstance* InAnimInstance) override;
 	virtual void UpdateInternal(const FAnimationUpdateContext& Context) override;
@@ -268,6 +267,7 @@ private:
 	uint8 bEnabled : 1;
 	uint8 bSimulationStarted : 1;
 	uint8 bCheckForBodyTransformInit : 1;
+	uint8 bInitialized : 1;
 
 public:
 	void PostSerialize(const FArchive& Ar);
