Fixed getter functions for the Sequence Player Library

diff --git a/Engine/Source/Runtime/AnimGraphRuntime/Private/SequencePlayerLibrary.cpp b/Engine/Source/Runtime/AnimGraphRuntime/Private/SequencePlayerLibrary.cpp
--- a/Engine/Source/Runtime/AnimGraphRuntime/Private/SequencePlayerLibrary.cpp
+++ b/Engine/Source/Runtime/AnimGraphRuntime/Private/SequencePlayerLibrary.cpp
@@ -102,67 +102,66 @@ FSequencePlayerReference USequencePlayerLibrary::SetSequenceWithInertialBlending
 	return SequencePlayer;
 }
 
+FSequencePlayerReference USequencePlayerLibrary::GetSequence(const FSequencePlayerReference& SequencePlayer, UAnimSequenceBase*& SequenceBase)
+{
+	SequencePlayer.CallAnimNodeFunction<FAnimNode_SequencePlayer>(
+		TEXT("GetSequence"),
+		[&SequenceBase](FAnimNode_SequencePlayer& InSequencePlayer)
+		{
+			SequenceBase = InSequencePlayer.GetSequence();
+		});
+
+	return SequencePlayer;
+}
+
 float USequencePlayerLibrary::GetAccumulatedTime(const FSequencePlayerReference& SequencePlayer)
 {
+	float AccumulatedTime = 0.f;
 	SequencePlayer.CallAnimNodeFunction<FAnimNode_SequencePlayer>(
 		TEXT("GetAccumulatedTime"),
-		[](FAnimNode_SequencePlayer& InSequencePlayer)
+		[&AccumulatedTime](FAnimNode_SequencePlayer& InSequencePlayer)
 		{
-			return InSequencePlayer.GetAccumulatedTime();
+			AccumulatedTime = InSequencePlayer.GetAccumulatedTime();
 		});
 
-	UE_LOG(LogSequencePlayerLibrary, Warning, TEXT("Could not get accumulated time on sequence player."));
-	return 0.f;
+	return AccumulatedTime;
 }
 
 float USequencePlayerLibrary::GetStartPosition(const FSequencePlayerReference& SequencePlayer)
 {
+	float StartPosition = 0.f;
 	SequencePlayer.CallAnimNodeFunction<FAnimNode_SequencePlayer>(
 		TEXT("GetStartPosition"),
-		[](FAnimNode_SequencePlayer& InSequencePlayer)
+		[&StartPosition](FAnimNode_SequencePlayer& InSequencePlayer)
 		{
-			return InSequencePlayer.GetStartPosition();
+			StartPosition = InSequencePlayer.GetStartPosition();
 		});
 
-	UE_LOG(LogSequencePlayerLibrary, Warning, TEXT("Could not get start position on sequence player."));
-	return 0.f;
+	return StartPosition;
 }
 
 float USequencePlayerLibrary::GetPlayRate(const FSequencePlayerReference& SequencePlayer)
 {
+	float PlayRate = 1.f;
 	SequencePlayer.CallAnimNodeFunction<FAnimNode_SequencePlayer>(
 		TEXT("GetPlayRate"),
-		[](FAnimNode_SequencePlayer& InSequencePlayer)
+		[&PlayRate](FAnimNode_SequencePlayer& InSequencePlayer)
 		{
-			return InSequencePlayer.GetPlayRate();
+			PlayRate = InSequencePlayer.GetPlayRate();
 		});
 
-	UE_LOG(LogSequencePlayerLibrary, Warning, TEXT("Could not get play rate on sequence player."));
-	return 1.f;
+	return PlayRate;
 }
 
 bool USequencePlayerLibrary::GetLoopAnimation(const FSequencePlayerReference& SequencePlayer)
 {
+	bool bLoopAnimation = false;
 	SequencePlayer.CallAnimNodeFunction<FAnimNode_SequencePlayer>(
 		TEXT("GetLoopAnimation"),
-		[](FAnimNode_SequencePlayer& InSequencePlayer)
-		{
-			return InSequencePlayer.GetLoopAnimation();
-		});
-
-	UE_LOG(LogSequencePlayerLibrary, Warning, TEXT("Could not get looping state on sequence player."));
-	return false;
-}
-
-UAnimSequenceBase* USequencePlayerLibrary::GetSequence(const FSequencePlayerReference& SequencePlayer)
-{
-	SequencePlayer.CallAnimNodeFunction<FAnimNode_SequencePlayer>(
-		TEXT("GetSequence"),
-		[](FAnimNode_SequencePlayer& InSequencePlayer)
+		[&bLoopAnimation](FAnimNode_SequencePlayer& InSequencePlayer)
 		{
-			return InSequencePlayer.GetSequence();
+			bLoopAnimation = InSequencePlayer.GetLoopAnimation();
 		});
 
-	UE_LOG(LogSequencePlayerLibrary, Warning, TEXT("Could not get sequence on sequence player."));
-	return nullptr;
+	return bLoopAnimation;
 }
\ No newline at end of file
diff --git a/Engine/Source/Runtime/AnimGraphRuntime/Public/SequencePlayerLibrary.h b/Engine/Source/Runtime/AnimGraphRuntime/Public/SequencePlayerLibrary.h
--- a/Engine/Source/Runtime/AnimGraphRuntime/Public/SequencePlayerLibrary.h
+++ b/Engine/Source/Runtime/AnimGraphRuntime/Public/SequencePlayerLibrary.h
@@ -60,6 +60,10 @@ public:
 	UFUNCTION(BlueprintCallable, Category = "Sequence Player", meta = (BlueprintThreadSafe))
 	static FSequencePlayerReference SetSequenceWithInertialBlending(const FAnimUpdateContext& UpdateContext, const FSequencePlayerReference& SequencePlayer, UAnimSequenceBase* Sequence, float BlendTime = 0.2f);
 	
+	/** Get the current sequence of the sequence player */
+	UFUNCTION(BlueprintCallable, Category = "Sequence Player", meta = (BlueprintThreadSafe))
+	static FSequencePlayerReference GetSequence(const FSequencePlayerReference& SequencePlayer, UPARAM(Ref) UAnimSequenceBase*& SequenceBase);
+
 	/** Gets the current accumulated time of the sequence player */
 	UFUNCTION(BlueprintCallable, Category = "Sequence Player", meta = (BlueprintThreadSafe))
 	static float GetAccumulatedTime(const FSequencePlayerReference& SequencePlayer);
@@ -75,8 +79,4 @@ public:
 	/** Get the looping state of the sequence player */
 	UFUNCTION(BlueprintCallable, Category = "Sequence Player", meta = (BlueprintThreadSafe))
 	static bool GetLoopAnimation(const FSequencePlayerReference& SequencePlayer);
-
-	/** Get the current sequence of the sequence player */
-	UFUNCTION(BlueprintCallable, Category = "Sequence Player", meta = (BlueprintThreadSafe))
-	static UAnimSequenceBase* GetSequence(const FSequencePlayerReference& SequencePlayer);
 };
\ No newline at end of file
