[UE][FEATURE] AnimDynamics node TRS Widget and Chain Body Editing:
+ Support for editing most geometric parameters (Joint Offsets, Box Extents, Collision Planes etc) with a TRS Widget.
+ Allow parameters to be edited for each individual physics body in a chain.
+ DEPRECATED some physics body properties, these are now held in an array, one instance per physics body.
+ Changed definition of LocalJointOffset - was joint position relative to physics body, now physics body position relative to joint
+ Fix for apparent bug in position constraints between chain bodies.
+ Added a new AnimDynamics Edit Mode to support new widget features.
+ Changes to details panel layout.

diff --git a/Engine/Source/Editor/AnimGraph/Private/AnimGraphModule.cpp b/Engine/Source/Editor/AnimGraph/Private/AnimGraphModule.cpp
--- a/Engine/Source/Editor/AnimGraph/Private/AnimGraphModule.cpp
+++ b/Engine/Source/Editor/AnimGraph/Private/AnimGraphModule.cpp
@@ -11,6 +11,7 @@
 #include "PropertyEditorModule.h"
 #include "AnimGraphNode_PoseDriver.h"
 #include "PoseDriverDetails.h"
+#include "EditModes/AnimDynamicsEditMode.h"
 #include "EditModes/TwoBoneIKEditMode.h"
 #include "EditModes/ObserveBoneEditMode.h"
 #include "EditModes/ModifyBoneEditMode.h"
@@ -59,6 +60,7 @@ void FAnimGraphModule::StartupModule()
 	IModularFeatures::Get().RegisterModularFeature("PropertyAccessBlueprintBinding", &PropertyAccessAnimBlueprintBinding);
 
 	// Register the editor modes
+	FEditorModeRegistry::Get().RegisterMode<FAnimDynamicsEditMode>(AnimNodeEditModes::AnimDynamics, LOCTEXT("AnimDynamicsEditMode", "Anim Dynamics"), FSlateIcon(), false);
 	FEditorModeRegistry::Get().RegisterMode<FAnimNodeEditMode>(AnimNodeEditModes::AnimNode, LOCTEXT("AnimNodeEditMode", "Anim Node"), FSlateIcon(), false);
 	FEditorModeRegistry::Get().RegisterMode<FTwoBoneIKEditMode>(AnimNodeEditModes::TwoBoneIK, LOCTEXT("TwoBoneIKEditMode", "2-Bone IK"), FSlateIcon(), false);
 	FEditorModeRegistry::Get().RegisterMode<FObserveBoneEditMode>(AnimNodeEditModes::ObserveBone, LOCTEXT("ObserveBoneEditMode", "Observe Bone"), FSlateIcon(), false);
diff --git a/Engine/Source/Editor/AnimGraph/Private/AnimGraphNode_AnimDynamics.cpp b/Engine/Source/Editor/AnimGraph/Private/AnimGraphNode_AnimDynamics.cpp
--- a/Engine/Source/Editor/AnimGraph/Private/AnimGraphNode_AnimDynamics.cpp
+++ b/Engine/Source/Editor/AnimGraph/Private/AnimGraphNode_AnimDynamics.cpp
@@ -1,9 +1,11 @@
 // Copyright Epic Games, Inc. All Rights Reserved.
 
 #include "AnimGraphNode_AnimDynamics.h"
+#include "AnimNodeEditModes.h"
 #include "EngineGlobals.h"
 #include "SceneManagement.h"
 #include "Materials/MaterialInstanceDynamic.h"
+#include "UObject/ReleaseObjectVersion.h"
 #include "Widgets/Input/SButton.h"
 
 // Details includes
@@ -21,228 +23,27 @@ FText UAnimGraphNode_AnimDynamics::GetTooltipText() const
 	return LOCTEXT("NodeTooltip", "Anim Dynamics");
 }
 
-void UAnimGraphNode_AnimDynamics::Draw(FPrimitiveDrawInterface* PDI, USkeletalMeshComponent * PreviewSkelMeshComp) const
-{
-	check(PreviewSkelMeshComp);
-
-	if(LastPreviewComponent != PreviewSkelMeshComp)
-	{
-		LastPreviewComponent = PreviewSkelMeshComp;
-	}
-
-	// If we want to preview the live node, process it here
-	if(bPreviewLive)
-	{
-		FAnimNode_AnimDynamics* ActivePreviewNode = GetPreviewDynamicsNode();
-
-		if(ActivePreviewNode)
-		{
-			for(int32 BodyIndex = 0 ; BodyIndex < ActivePreviewNode->GetNumBodies() ; ++BodyIndex)
-			{
-				const FAnimPhysRigidBody& Body = ActivePreviewNode->GetPhysBody(BodyIndex);
-				FTransform BodyTransform(Body.Pose.Orientation, Body.Pose.Position);
-
-				// Physics bodies are in Simulation Space. Transform into component space before rendering in the viewport.
-				if (ActivePreviewNode->SimulationSpace == AnimPhysSimSpaceType::RootRelative)
-				{
-					const FTransform RelativeBoneTransform = PreviewSkelMeshComp->GetBoneTransform(0);
-					BodyTransform = BodyTransform * RelativeBoneTransform;
-				}
-				else if (ActivePreviewNode->SimulationSpace == AnimPhysSimSpaceType::BoneRelative)
-				{
-					const FTransform RelativeBoneTransform = PreviewSkelMeshComp->GetBoneTransform(PreviewSkelMeshComp->GetBoneIndex(ActivePreviewNode->RelativeSpaceBone.BoneName));
-					BodyTransform = BodyTransform * RelativeBoneTransform;
-				}
-
-				for(const FAnimPhysShape& Shape : Body.Shapes)
-				{
-					for(const FIntVector& Triangle : Shape.Triangles)
-					{
-						for(int32 Idx = 0 ; Idx < 3 ; ++Idx)
-						{
-							int32 Next = (Idx + 1) % 3;
-
-							FVector FirstVertPosition = BodyTransform.TransformPosition(Shape.Vertices[Triangle[Idx]]);
-							FVector SecondVertPosition = BodyTransform.TransformPosition(Shape.Vertices[Triangle[Next]]);
-
-							PDI->DrawLine(FirstVertPosition, SecondVertPosition, AnimDynamicsNodeConstants::ActiveBodyDrawColor, SDPG_Foreground, AnimDynamicsNodeConstants::ShapeLineWidth);
-						}
-					}
-				}
-				
-				const int32 BoneIndex = PreviewSkelMeshComp->GetBoneIndex(ActivePreviewNode->BoundBone.BoneName);
-				if(BoneIndex != INDEX_NONE)
-				{
-					FTransform BodyJointTransform = PreviewSkelMeshComp->GetBoneTransform(BoneIndex);
-					FTransform ShapeOriginalTransform = BodyJointTransform;
-
-					// Draw pin location
-					FVector LocalPinOffset = BodyTransform.Rotator().RotateVector(Node.GetBodyLocalJointOffset(BodyIndex));
-					PDI->DrawLine(Body.Pose.Position, Body.Pose.Position + LocalPinOffset, FLinearColor::Green, SDPG_Foreground, AnimDynamicsNodeConstants::ShapeLineWidth);
-
-					// Draw basis at body location
-					FVector Origin = BodyTransform.GetTranslation();
-					FVector XAxis(1.0f, 0.0f, 0.0f);
-					FVector YAxis(0.0f, 1.0f, 0.0f);
-					FVector ZAxis(0.0f, 0.0f, 1.0f);
-
-					XAxis = BodyTransform.TransformVector(XAxis);
-					YAxis = BodyTransform.TransformVector(YAxis);
-					ZAxis = BodyTransform.TransformVector(ZAxis);
-
-					PDI->DrawLine(Origin, Origin + XAxis * AnimDynamicsNodeConstants::TransformBasisScale, FLinearColor::Red, SDPG_Foreground, AnimDynamicsNodeConstants::TransformLineWidth);
-					PDI->DrawLine(Origin, Origin + YAxis * AnimDynamicsNodeConstants::TransformBasisScale, FLinearColor::Green, SDPG_Foreground, AnimDynamicsNodeConstants::TransformLineWidth);
-					PDI->DrawLine(Origin, Origin + ZAxis * AnimDynamicsNodeConstants::TransformBasisScale, FLinearColor::Blue, SDPG_Foreground, AnimDynamicsNodeConstants::TransformLineWidth);
-
-					if(bShowLinearLimits)
-					{
-						DrawLinearLimits(PDI, BodyJointTransform, *ActivePreviewNode);
-					}
-
-					if(bShowAngularLimits)
-					{
-						FTransform AngularLimitsTM(BodyJointTransform.GetRotation(), BodyTransform.GetTranslation() + LocalPinOffset);
-						DrawAngularLimits(PDI, AngularLimitsTM, *ActivePreviewNode);
-					}
-
-					if(bShowCollisionSpheres && Body.CollisionType != AnimPhysCollisionType::CoM)
-					{
-						// Draw collision sphere
-						DrawWireSphere(PDI, BodyTransform, FLinearColor(FColor::Cyan), Body.SphereCollisionRadius, 24, SDPG_Foreground, 0.2f);
-					}
-				}
-			}
-
-			// Only draw the planar limit once
-			if(bShowPlanarLimit && ActivePreviewNode->PlanarLimits.Num() > 0)
-			{
-				for(FAnimPhysPlanarLimit& PlanarLimit : ActivePreviewNode->PlanarLimits)
-				{
-					FTransform LimitPlaneTransform = PlanarLimit.PlaneTransform;
-					const int32 LimitDrivingBoneIdx = PreviewSkelMeshComp->GetBoneIndex(PlanarLimit.DrivingBone.BoneName);
-
-					if(LimitDrivingBoneIdx != INDEX_NONE)
-					{
-						LimitPlaneTransform *= PreviewSkelMeshComp->GetComponentSpaceTransforms()[LimitDrivingBoneIdx];
-					}
-
-					DrawPlane10x10(PDI, LimitPlaneTransform.ToMatrixNoScale(), 200.0f, FVector2D(0.0f, 0.0f), FVector2D(1.0f, 1.0f), GEngine->ConstraintLimitMaterialPrismatic->GetRenderProxy(), SDPG_World);
-					DrawDirectionalArrow(PDI, FRotationMatrix(FRotator(90.0f, 0.0f, 0.0f)) * LimitPlaneTransform.ToMatrixNoScale(), FLinearColor::Blue, 50.0f, 20.0f, SDPG_Foreground, 0.5f);
-				}
-			}
-
-			if(bShowSphericalLimit && ActivePreviewNode->SphericalLimits.Num() > 0)
-			{
-				for(FAnimPhysSphericalLimit& SphericalLimit : ActivePreviewNode->SphericalLimits)
-				{
-					FTransform SphereTransform = FTransform::Identity;
-					SphereTransform.SetTranslation(SphericalLimit.SphereLocalOffset);
-					
-					const int32 DrivingBoneIdx = PreviewSkelMeshComp->GetBoneIndex(SphericalLimit.DrivingBone.BoneName);
-
-					if(DrivingBoneIdx != INDEX_NONE)
-					{
-						SphereTransform *= PreviewSkelMeshComp->GetComponentSpaceTransforms()[DrivingBoneIdx];
-					}
-
-					DrawSphere(PDI, SphereTransform.GetLocation(), FRotator::ZeroRotator, FVector(SphericalLimit.LimitRadius), 24, 6, GEngine->ConstraintLimitMaterialPrismatic->GetRenderProxy(), SDPG_World);
-					DrawWireSphere(PDI, SphereTransform, FLinearColor::Black, SphericalLimit.LimitRadius, 24, SDPG_World);
-				}
-			}
-		}
-	}
-	else
-	{
-		const int32 BoneIndex = PreviewSkelMeshComp->GetBoneIndex(Node.BoundBone.BoneName);
-
-		if(BoneIndex != INDEX_NONE)
-		{
-			// World space transform
-			const FTransform BoneTransform = PreviewSkelMeshComp->GetBoneTransform(BoneIndex);
-			FTransform ShapeTransform = BoneTransform;
-			ShapeTransform.SetTranslation(ShapeTransform.GetTranslation() - BoneTransform.GetRotation().RotateVector(Node.LocalJointOffset));
-			
-			for(const FIntVector& Triangle : EditPreviewShape.Triangles)
-			{
-				for(int32 Idx = 0 ; Idx < 3 ; ++Idx)
-				{
-					int32 Next = (Idx + 1) % 3;
-
-					FVector FirstVertPosition = ShapeTransform.TransformPosition(EditPreviewShape.Vertices[Triangle[Idx]]);
-					FVector SecondVertPosition = ShapeTransform.TransformPosition(EditPreviewShape.Vertices[Triangle[Next]]);
-
-					PDI->DrawLine(FirstVertPosition, SecondVertPosition, AnimDynamicsNodeConstants::ShapeDrawColor, SDPG_Foreground, AnimDynamicsNodeConstants::ShapeLineWidth);
-				}
-			}
-
-			// COM
-			FVector Origin = ShapeTransform.GetTranslation();
-			FVector XAxis(1.0f, 0.0f, 0.0f);
-			FVector YAxis(0.0f, 1.0f, 0.0f);
-			FVector ZAxis(0.0f, 0.0f, 1.0f);
-
-			XAxis = ShapeTransform.TransformVector(XAxis);
-			YAxis = ShapeTransform.TransformVector(YAxis);
-			ZAxis = ShapeTransform.TransformVector(ZAxis);
-
-			PDI->DrawLine(Origin, Origin + XAxis * AnimDynamicsNodeConstants::TransformBasisScale, FLinearColor::Red, SDPG_Foreground, 0.5f);
-			PDI->DrawLine(Origin, Origin + YAxis * AnimDynamicsNodeConstants::TransformBasisScale, FLinearColor::Green, SDPG_Foreground, 0.5f);
-			PDI->DrawLine(Origin, Origin + ZAxis * AnimDynamicsNodeConstants::TransformBasisScale, FLinearColor::Blue, SDPG_Foreground, 0.5f);
-
-			// Local joint offset
-			FVector JointOffset = ShapeTransform.Rotator().RotateVector(Node.LocalJointOffset);
-			PDI->DrawLine(ShapeTransform.GetTranslation(), ShapeTransform.GetTranslation() + JointOffset, FLinearColor::Green, SDPG_Foreground, AnimDynamicsNodeConstants::ShapeLineWidth);
-
-			if (bShowLinearLimits)
-			{
-				DrawLinearLimits(PDI, ShapeTransform, Node);
-			}
-
-			if (bShowAngularLimits)
-			{
-				DrawAngularLimits(PDI, ShapeTransform, Node);
-			}
-		}
-	}
-}
-
 void UAnimGraphNode_AnimDynamics::GetOnScreenDebugInfo(TArray<FText>& DebugInfo, FAnimNode_Base* RuntimeAnimNode, USkeletalMeshComponent* PreviewSkelMeshComp) const
 {
 	if(RuntimeAnimNode)
 	{
-		FAnimNode_AnimDynamics* PreviewNode = static_cast<FAnimNode_AnimDynamics*>(RuntimeAnimNode);
-		int32 NumBones = PreviewNode->GetNumBoundBones();
-		for(int32 ChainBoneIndex = 0; ChainBoneIndex < NumBones; ++ChainBoneIndex)
+		const FAnimNode_AnimDynamics* PreviewNode = static_cast<FAnimNode_AnimDynamics*>(RuntimeAnimNode);
+
+		for(const FAnimPhysBodyDefinition& PhysicsBodyDef : PreviewNode->PhysicsBodyDefinitions)
 		{
-			if(const FBoneReference* BoneRef = PreviewNode->GetBoundBoneReference(ChainBoneIndex))
+			const FName BoneName = PhysicsBodyDef.BoundBone.BoneName;
+			const int32 SkelBoneIndex = PreviewSkelMeshComp->GetBoneIndex(BoneName);
+			if(SkelBoneIndex != INDEX_NONE)
 			{
-				const int32 SkelBoneIndex = PreviewSkelMeshComp->GetBoneIndex(BoneRef->BoneName);
-				if(SkelBoneIndex != INDEX_NONE)
-				{
-					FTransform BoneTransform = PreviewSkelMeshComp->GetBoneTransform(SkelBoneIndex);
-					DebugInfo.Add(FText::Format(LOCTEXT("DebugOnScreenName", "Anim Dynamics (Bone:{0})"), FText::FromName(BoneRef->BoneName)));
-					DebugInfo.Add(FText::Format(LOCTEXT("DebugOnScreenTranslation", "    Translation: {0}"), FText::FromString(BoneTransform.GetTranslation().ToString())));
-					DebugInfo.Add(FText::Format(LOCTEXT("DebugOnScreenRotation", "    Rotation: {0}"), FText::FromString(BoneTransform.Rotator().ToString())));
-				}
+				FTransform BoneTransform = PreviewSkelMeshComp->GetBoneTransform(SkelBoneIndex);
+				DebugInfo.Add(FText::Format(LOCTEXT("DebugOnScreenName", "Anim Dynamics (Bone:{0})"), FText::FromName(BoneName)));
+				DebugInfo.Add(FText::Format(LOCTEXT("DebugOnScreenTranslation", "    Translation: {0}"), FText::FromString(BoneTransform.GetTranslation().ToString())));
+				DebugInfo.Add(FText::Format(LOCTEXT("DebugOnScreenRotation", "    Rotation: {0}"), FText::FromString(BoneTransform.Rotator().ToString())));
 			}
 		}
 	}
 }
 
-void UAnimGraphNode_AnimDynamics::DrawLinearLimits(FPrimitiveDrawInterface* PDI, FTransform ShapeTransform, const FAnimNode_AnimDynamics& NodeToVisualise) const
-{
-	// Draw linear limits
-	FVector LinearLimitHalfExtents(NodeToVisualise.ConstraintSetup.LinearAxesMax - NodeToVisualise.ConstraintSetup.LinearAxesMin);
-	// Add a tiny bit so we can see collapsed axes
-	LinearLimitHalfExtents += FVector(0.1f);
-	LinearLimitHalfExtents /= 2.0f;
-	FVector LinearLimitsCenter = NodeToVisualise.ConstraintSetup.LinearAxesMin + LinearLimitHalfExtents;
-	FTransform LinearLimitsTransform = ShapeTransform;
-	LinearLimitsTransform.SetTranslation(LinearLimitsTransform.GetTranslation() + LinearLimitsTransform.TransformVector(LinearLimitsCenter));
-
-	DrawBox(PDI, LinearLimitsTransform.ToMatrixWithScale(), LinearLimitHalfExtents, GEngine->ConstraintLimitMaterialPrismatic->GetRenderProxy(), SDPG_Foreground);
-}
-
 FText UAnimGraphNode_AnimDynamics::GetControllerDescription() const
 {
 	return LOCTEXT("Description", "Anim Dynamics");
@@ -260,19 +61,62 @@ void UAnimGraphNode_AnimDynamics::CustomizeDetails(IDetailLayoutBuilder& DetailB
 	FDetailWidgetRow& WidgetRow = PreviewCategory.AddCustomRow(LOCTEXT("ResetButtonRow", "Reset"));
 
 	WidgetRow
-	[
-		SNew(SButton)
+		[
+			SNew(SButton)
 			.Text(LOCTEXT("ResetButtonText", "Reset Simulation"))
 			.ToolTipText(LOCTEXT("ResetButtonToolTip", "Resets the simulation for this node"))
 			.OnClicked(FOnClicked::CreateStatic(&UAnimGraphNode_AnimDynamics::ResetButtonClicked, &DetailBuilder))
-	];
+		];
+
+	// Add warning message about physics body array not being updated untill after the first BP compile.
+	if (LastPreviewComponent == nullptr)
+	{
+		IDetailCategoryBuilder& SetupCategory = DetailBuilder.EditCategory(TEXT("Setup"));
+		const FText WarningText(LOCTEXT("AnimDynamicsWarningText", "WARNING - Physics Bodies Will Not Be Valid Untill This Node Has Been Connected And Compiled"));
+		FDetailWidgetRow& WarningMessageTextWidgetRow = SetupCategory.AddCustomRow(WarningText);
+
+		WarningMessageTextWidgetRow
+			[
+				SNew(SHorizontalBox)
+				+ SHorizontalBox::Slot()
+				.AutoWidth()
+				[
+					SNew(STextBlock)
+					.Text(WarningText)
+					.Justification(ETextJustify::Center)
+					.ColorAndOpacity(FLinearColor::Red)
+				]
+			];
+	}
+
+	// Force order of details panel catagories - Must set order for all of them as any that are edited automatically move to the top.
+	{
+		uint32 SortOrder = 0;
+		DetailBuilder.EditCategory("Preview").SetSortOrder(SortOrder++);
+		DetailBuilder.EditCategory("Setup").SetSortOrder(SortOrder++);
+		DetailBuilder.EditCategory("Settings").SetSortOrder(SortOrder++);
+		DetailBuilder.EditCategory("SphericalLimit").SetSortOrder(SortOrder++);
+		DetailBuilder.EditCategory("PlanarLimit").SetSortOrder(SortOrder++);
+		DetailBuilder.EditCategory("Forces").SetSortOrder(SortOrder++);
+		DetailBuilder.EditCategory("Wind").SetSortOrder(SortOrder++);
+		DetailBuilder.EditCategory("Retargetting").SetSortOrder(SortOrder++);
+		DetailBuilder.EditCategory("Performance").SetSortOrder(SortOrder++);
+		DetailBuilder.EditCategory("Functions").SetSortOrder(SortOrder++);
+		DetailBuilder.EditCategory("Alpha").SetSortOrder(SortOrder++);
+	}
 }
 
+
 void UAnimGraphNode_AnimDynamics::ValidateAnimNodeDuringCompilation(USkeleton* ForSkeleton, FCompilerResultsLog& MessageLog)
 {
 	Super::ValidateAnimNodeDuringCompilation(ForSkeleton, MessageLog);
 }
 
+FEditorModeID UAnimGraphNode_AnimDynamics::GetEditorMode() const
+{
+	return AnimNodeEditModes::AnimDynamics;
+}
+
 FText UAnimGraphNode_AnimDynamics::GetNodeTitle(ENodeTitleType::Type TitleType) const
 {
 	FFormatNamedArguments Arguments;
@@ -316,50 +160,33 @@ FText UAnimGraphNode_AnimDynamics::GetNodeTitle(ENodeTitleType::Type TitleType)
 
 void UAnimGraphNode_AnimDynamics::PostEditChangeProperty(struct FPropertyChangedEvent& PropertyChangedEvent)
 {
-	Super::PostEditChangeProperty(PropertyChangedEvent);
-
-	// Regenerate render shape(s)
-	EditPreviewShape = FAnimPhysShape::MakeBox(Node.BoxExtents);
-}
-
-void UAnimGraphNode_AnimDynamics::DrawAngularLimits(FPrimitiveDrawInterface* PDI, FTransform JointTransform, const FAnimNode_AnimDynamics& NodeToVisualize) const
-{
-	FVector XAxis = JointTransform.GetUnitAxis(EAxis::X);
-	FVector YAxis = JointTransform.GetUnitAxis(EAxis::Y);
-	FVector ZAxis = JointTransform.GetUnitAxis(EAxis::Z);
-
-	const FVector& MinAngles = NodeToVisualize.ConstraintSetup.AngularLimitsMin;
-	const FVector& MaxAngles = NodeToVisualize.ConstraintSetup.AngularLimitsMax;
-	FVector AngleRange = MaxAngles - MinAngles;
-	FVector Middle = MinAngles + AngleRange * 0.5f;
-
-	if (AngleRange.X > 0.0f && AngleRange.X < 180.0f)
+	if (PropertyChangedEvent.GetPropertyName() == GET_MEMBER_NAME_CHECKED(FAnimNode_AnimDynamics, ChainEnd))
 	{
-		FTransform XAxisConeTM(YAxis, XAxis ^ YAxis, XAxis, JointTransform.GetTranslation());
-		XAxisConeTM.SetRotation(FQuat(XAxis, FMath::DegreesToRadians(-Middle.X)) * XAxisConeTM.GetRotation());
-		DrawCone(PDI, FScaleMatrix(30.0f) * XAxisConeTM.ToMatrixWithScale(), FMath::DegreesToRadians(AngleRange.X / 2.0f), 0.0f, 24, false, FLinearColor::White, GEngine->ConstraintLimitMaterialX->GetRenderProxy(), SDPG_World);
-	}
+		// bChain has been modified.
+		if (Node.bChain)
+		{
+			Node.ChainEnd = Node.PhysicsBodyDefinitions[0].BoundBone;
+		}
+		else
+		{
+			Node.ChainEnd.BoneName = NAME_None;
+		}
 
-	if (AngleRange.Y > 0.0f && AngleRange.Y < 180.0f)
-	{
-		FTransform YAxisConeTM(ZAxis, YAxis ^ ZAxis, YAxis, JointTransform.GetTranslation());
-		YAxisConeTM.SetRotation(FQuat(YAxis, FMath::DegreesToRadians(Middle.Y)) * YAxisConeTM.GetRotation());
-		DrawCone(PDI, FScaleMatrix(30.0f) * YAxisConeTM.ToMatrixWithScale(), FMath::DegreesToRadians(AngleRange.Y / 2.0f), 0.0f, 24, false, FLinearColor::White, GEngine->ConstraintLimitMaterialY->GetRenderProxy(), SDPG_World);
+		Node.UpdateChainPhysicsBodyDefinitions(LastPreviewComponent);
 	}
 
-	if (AngleRange.Z > 0.0f && AngleRange.Z < 180.0f)
+	if (PropertyChangedEvent.GetPropertyName() == GET_MEMBER_NAME_CHECKED(FBoneReference, BoneName))
 	{
-		FTransform ZAxisConeTM(XAxis, ZAxis ^ XAxis, ZAxis, JointTransform.GetTranslation());
-		ZAxisConeTM.SetRotation(FQuat(ZAxis, FMath::DegreesToRadians(Middle.Z)) * ZAxisConeTM.GetRotation());
-		DrawCone(PDI, FScaleMatrix(30.0f) * ZAxisConeTM.ToMatrixWithScale(), FMath::DegreesToRadians(AngleRange.Z / 2.0f), 0.0f, 24, false, FLinearColor::White, GEngine->ConstraintLimitMaterialZ->GetRenderProxy(), SDPG_World);
+		// Either BoundBone or ChainEnd have been modified.
+		Node.UpdateChainPhysicsBodyDefinitions(LastPreviewComponent);
 	}
+
+	Super::PostEditChangeProperty(PropertyChangedEvent);
 }
 
 void UAnimGraphNode_AnimDynamics::PostLoad()
 {
 	Super::PostLoad();
-
-	EditPreviewShape = FAnimPhysShape::MakeBox(Node.BoxExtents);
 }
 
 void UAnimGraphNode_AnimDynamics::ResetSim()
@@ -411,7 +238,7 @@ void UAnimGraphNode_AnimDynamics::Serialize(FArchive& Ar)
 	
 	if(CustomAnimVersion < FAnimationCustomVersion::AnimDynamicsAddAngularOffsets)
 	{
-		FAnimPhysConstraintSetup& ConSetup = Node.ConstraintSetup;
+		FAnimPhysConstraintSetup& ConSetup = Node.ConstraintSetup_DEPRECATED;
 		ConSetup.AngularLimitsMin = FVector(-ConSetup.AngularXAngle_DEPRECATED, -ConSetup.AngularYAngle_DEPRECATED, -ConSetup.AngularZAngle_DEPRECATED);
 		ConSetup.AngularLimitsMax = FVector(ConSetup.AngularXAngle_DEPRECATED, ConSetup.AngularYAngle_DEPRECATED, ConSetup.AngularZAngle_DEPRECATED);
 	}
@@ -422,6 +249,20 @@ void UAnimGraphNode_AnimDynamics::Serialize(FArchive& Ar)
 	{
 		Node.bGravityOverrideInSimSpace = true;
 	}
+
+	if (Ar.CustomVer(FFortniteMainBranchObjectVersion::GUID) < FFortniteMainBranchObjectVersion::AnimDynamicsEditableChainParameters)
+	{
+		// Initialise first physics body using deprecated parameters then re-initialize chain.
+		Node.PhysicsBodyDefinitions.Reset();		
+		FAnimPhysBodyDefinition PhysBodyDef;
+		PhysBodyDef.BoundBone = Node.BoundBone;
+		PhysBodyDef.BoxExtents = Node.BoxExtents_DEPRECATED;
+		PhysBodyDef.LocalJointOffset = -Node.LocalJointOffset_DEPRECATED; // Note: definition of joint offset has changed from 'Joint position relative to physics body' to 'physics body position relative to Joint'.
+		PhysBodyDef.ConstraintSetup = Node.ConstraintSetup_DEPRECATED;
+		PhysBodyDef.CollisionType = Node.CollisionType_DEPRECATED;
+		PhysBodyDef.SphereCollisionRadius = Node.SphereCollisionRadius_DEPRECATED;
+		Node.PhysicsBodyDefinitions.Add(PhysBodyDef);
+	}
 }
 
 #undef LOCTEXT_NAMESPACE
diff --git a/Engine/Source/Editor/AnimGraph/Private/EditModes/AnimDynamicsEditMode.cpp b/Engine/Source/Editor/AnimGraph/Private/EditModes/AnimDynamicsEditMode.cpp
new file mode 100644
--- /dev/null
+++ b/Engine/Source/Editor/AnimGraph/Private/EditModes/AnimDynamicsEditMode.cpp
@@ -0,0 +1,850 @@
+// Copyright Epic Games, Inc. All Rights Reserved.
+
+#include "EditModes/AnimDynamicsEditMode.h"
+#include "AnimGraphNode_AnimDynamics.h"
+#include "IPersonaPreviewScene.h"
+#include "Animation/DebugSkelMeshComponent.h"
+#include "EditorViewportClient.h"
+
+// AnimDynamicsEditMode hit proxy
+//
+// Allow users to select objects in the viewport for editing with a TRS widget.
+//
+struct HAnimDynamicsEditModeHitProxy : public HHitProxy
+{
+	DECLARE_HIT_PROXY()
+
+	FAnimDynamicsViewportObjectReference ViewportObjectRef;
+
+	HAnimDynamicsEditModeHitProxy(const uint32 InEditorNodeUniqueId, const FAnimDynamicsViewportObjectType InType, const uint32 InIndex) : HHitProxy(HPP_World), ViewportObjectRef(InEditorNodeUniqueId, InType, InIndex) {}
+	HAnimDynamicsEditModeHitProxy(const FAnimDynamicsViewportObjectReference& InGeometricObjectRef) : HHitProxy(HPP_World), ViewportObjectRef(InGeometricObjectRef) {}
+};
+
+IMPLEMENT_HIT_PROXY(HAnimDynamicsEditModeHitProxy, HHitProxy);
+
+// Class FAnimDynamicsViewportObjectReference
+FAnimDynamicsViewportObjectReference::FAnimDynamicsViewportObjectReference(const uint32 InEditorNodeUniqueId, const FAnimDynamicsViewportObjectType InType, const uint32 InIndex)
+	: EditorNodeUniqueId(InEditorNodeUniqueId)
+	, Type(InType)
+	, Index(InIndex)
+{}
+
+const bool operator==(const FAnimDynamicsViewportObjectReference& Lhs, const FAnimDynamicsViewportObjectReference& Rhs)
+{
+	return
+		(Lhs.EditorNodeUniqueId == Rhs.EditorNodeUniqueId) &&
+		(Lhs.Type == Rhs.Type) &&
+		(Lhs.Index == Rhs.Index);
+}
+
+// Non-member utility functions
+void DrawAngularLimits(FPrimitiveDrawInterface* PDI, FTransform JointTransform, const FAnimPhysConstraintSetup& ConstraintSetup)
+{
+	FVector XAxis = JointTransform.GetUnitAxis(EAxis::X);
+	FVector YAxis = JointTransform.GetUnitAxis(EAxis::Y);
+	FVector ZAxis = JointTransform.GetUnitAxis(EAxis::Z);
+
+	const FVector& MinAngles = ConstraintSetup.AngularLimitsMin;
+	const FVector& MaxAngles = ConstraintSetup.AngularLimitsMax;
+	FVector AngleRange = MaxAngles - MinAngles;
+	FVector Middle = MinAngles + AngleRange * 0.5f;
+
+	if (AngleRange.X > 0.0f && AngleRange.X < 180.0f)
+	{
+		FTransform XAxisConeTM(YAxis, XAxis ^ YAxis, XAxis, JointTransform.GetTranslation());
+		XAxisConeTM.SetRotation(FQuat(XAxis, FMath::DegreesToRadians(-Middle.X)) * XAxisConeTM.GetRotation());
+		DrawCone(PDI, FScaleMatrix(30.0f) * XAxisConeTM.ToMatrixWithScale(), FMath::DegreesToRadians(AngleRange.X / 2.0f), 0.0f, 24, false, FLinearColor::White, GEngine->ConstraintLimitMaterialX->GetRenderProxy(), SDPG_World);
+	}
+
+	if (AngleRange.Y > 0.0f && AngleRange.Y < 180.0f)
+	{
+		FTransform YAxisConeTM(ZAxis, YAxis ^ ZAxis, YAxis, JointTransform.GetTranslation());
+		YAxisConeTM.SetRotation(FQuat(YAxis, FMath::DegreesToRadians(Middle.Y)) * YAxisConeTM.GetRotation());
+		DrawCone(PDI, FScaleMatrix(30.0f) * YAxisConeTM.ToMatrixWithScale(), FMath::DegreesToRadians(AngleRange.Y / 2.0f), 0.0f, 24, false, FLinearColor::White, GEngine->ConstraintLimitMaterialY->GetRenderProxy(), SDPG_World);
+	}
+
+	if (AngleRange.Z > 0.0f && AngleRange.Z < 180.0f)
+	{
+		FTransform ZAxisConeTM(XAxis, ZAxis ^ XAxis, ZAxis, JointTransform.GetTranslation());
+		ZAxisConeTM.SetRotation(FQuat(ZAxis, FMath::DegreesToRadians(Middle.Z)) * ZAxisConeTM.GetRotation());
+		DrawCone(PDI, FScaleMatrix(30.0f) * ZAxisConeTM.ToMatrixWithScale(), FMath::DegreesToRadians(AngleRange.Z / 2.0f), 0.0f, 24, false, FLinearColor::White, GEngine->ConstraintLimitMaterialZ->GetRenderProxy(), SDPG_World);
+	}
+}
+
+void DrawLinearLimits(FPrimitiveDrawInterface* PDI, FTransform ShapeTransform, const FAnimPhysConstraintSetup& ConstraintSetup)
+{
+	// Draw linear limits
+	FVector LinearLimitHalfExtents(ConstraintSetup.LinearAxesMax - ConstraintSetup.LinearAxesMin);
+	// Add a tiny bit so we can see collapsed axes
+	LinearLimitHalfExtents += FVector(0.1f);
+	LinearLimitHalfExtents /= 2.0f;
+	FVector LinearLimitsCenter = ConstraintSetup.LinearAxesMin + LinearLimitHalfExtents;
+	FTransform LinearLimitsTransform = ShapeTransform;
+	LinearLimitsTransform.SetTranslation(LinearLimitsTransform.GetTranslation() + LinearLimitsTransform.TransformVector(LinearLimitsCenter));
+
+	DrawBox(PDI, LinearLimitsTransform.ToMatrixWithScale(), LinearLimitHalfExtents, GEngine->ConstraintLimitMaterialPrismatic->GetRenderProxy(), SDPG_Foreground);
+}
+
+template< typename TShape > void DrawShape(FPrimitiveDrawInterface* const PDI, const TShape& Shape, const FTransform& Transform, const FLinearColor& Color, const float LineWidth)
+{
+	for (const FIntVector& Triangle : Shape.Triangles)
+	{
+		for (int32 Idx = 0; Idx < 3; ++Idx)
+		{
+			int32 Next = (Idx + 1) % 3;
+
+			FVector FirstVertPosition = Transform.TransformPosition(Shape.Vertices[Triangle[Idx]]);
+			FVector SecondVertPosition = Transform.TransformPosition(Shape.Vertices[Triangle[Next]]);
+
+			PDI->DrawLine(FirstVertPosition, SecondVertPosition, Color, SDPG_Foreground, LineWidth);
+		}
+	}
+}
+
+void DrawBasis(FPrimitiveDrawInterface* const PDI, const FTransform& Transform)
+{
+	const FVector Origin = Transform.GetTranslation();
+	PDI->DrawLine(Origin, Origin + Transform.TransformVector(FVector::XAxisVector) * AnimDynamicsNodeConstants::TransformBasisScale, FLinearColor::Red, SDPG_Foreground, AnimDynamicsNodeConstants::TransformLineWidth);
+	PDI->DrawLine(Origin, Origin + Transform.TransformVector(FVector::YAxisVector) * AnimDynamicsNodeConstants::TransformBasisScale, FLinearColor::Green, SDPG_Foreground, AnimDynamicsNodeConstants::TransformLineWidth);
+	PDI->DrawLine(Origin, Origin + Transform.TransformVector(FVector::ZAxisVector) * AnimDynamicsNodeConstants::TransformBasisScale, FLinearColor::Blue, SDPG_Foreground, AnimDynamicsNodeConstants::TransformLineWidth);
+}
+
+
+// Class FAnimDynamicsEditMode
+FAnimDynamicsEditMode::FAnimDynamicsEditMode()
+	: CurWidgetMode(UE::Widget::WM_Translate)
+	, bIsInteractingWithWidget(false)
+{}
+
+void FAnimDynamicsEditMode::EnterMode(class UAnimGraphNode_Base* InEditorNode, struct FAnimNode_Base* InRuntimeNode)
+{
+	FAnimNodeEditMode::EnterMode(InEditorNode, InRuntimeNode);
+
+	// Cache pointer to PreviewSkelMesh in the Editor node as there's no good way to access it directly from the Editor Node and its required for some parts of the UI.
+	{
+		USkeletalMeshComponent* const PreviewSkelMeshComp = GetAnimPreviewScene().GetPreviewMeshComponent();
+
+		if (UAnimGraphNode_AnimDynamics* const EditorAnimDynamicsNode = CastChecked<UAnimGraphNode_AnimDynamics>(InEditorNode))
+		{
+			EditorAnimDynamicsNode->LastPreviewComponent = PreviewSkelMeshComp;
+		}
+	}
+}
+
+void FAnimDynamicsEditMode::ExitMode()
+{
+	SelectedViewportObjects.Empty();
+
+	for (const EditorRuntimeNodePair& CurrentNodePair : AnimNodes)
+	{
+		if (const UAnimGraphNode_AnimDynamics* const EditorAnimDynamicsNode = CastChecked<UAnimGraphNode_AnimDynamics>(CurrentNodePair.EditorAnimNode))
+		{
+			if (FAnimNode_AnimDynamics* const ActivePreviewNode = EditorAnimDynamicsNode->GetPreviewDynamicsNode())
+			{
+				ActivePreviewNode->bDoPhysicsUpdateInEditor = true; // Ensure physics update is enabled on premature exit.
+			}
+		}
+	}
+
+	FAnimNodeEditMode::ExitMode();
+}
+
+void FAnimDynamicsEditMode::Render(const FSceneView* View, FViewport* Viewport, FPrimitiveDrawInterface* PDI)
+{
+	USkeletalMeshComponent* const PreviewSkelMeshComp = GetAnimPreviewScene().GetPreviewMeshComponent();
+
+	check(View);
+	check(Viewport);
+	check(PDI);
+	check(PreviewSkelMeshComp);
+
+	for (const EditorRuntimeNodePair& CurrentNodePair : AnimNodes)
+	{
+		const UAnimGraphNode_AnimDynamics* const EditorAnimDynamicsNode = CastChecked<UAnimGraphNode_AnimDynamics>(CurrentNodePair.EditorAnimNode);
+
+		if (EditorAnimDynamicsNode)
+		{
+			if (EditorAnimDynamicsNode->LastPreviewComponent != PreviewSkelMeshComp)
+			{
+				EditorAnimDynamicsNode->LastPreviewComponent = PreviewSkelMeshComp;
+			}
+
+			const FAnimNode_AnimDynamics* const ActivePreviewNode = EditorAnimDynamicsNode->GetPreviewDynamicsNode();
+
+			if (!ActivePreviewNode)
+			{
+				return;
+			}
+
+			if (EditorAnimDynamicsNode->IsPreviewLiveActive())
+			{
+				for (int32 BodyIndex = 0; BodyIndex < ActivePreviewNode->GetNumBodies(); ++BodyIndex)
+				{
+					const FAnimPhysRigidBody& Body = ActivePreviewNode->GetPhysBody(BodyIndex);
+					FTransform BodyTransform(Body.Pose.Orientation, Body.Pose.Position);
+
+					// Physics bodies are in Simulation Space. Transform into component space before rendering in the viewport.
+					if (ActivePreviewNode->SimulationSpace == AnimPhysSimSpaceType::RootRelative)
+					{
+						const FTransform RelativeBoneTransform = PreviewSkelMeshComp->GetBoneTransform(0);
+						BodyTransform = BodyTransform * RelativeBoneTransform;
+					}
+					else if (ActivePreviewNode->SimulationSpace == AnimPhysSimSpaceType::BoneRelative)
+					{
+						const FTransform RelativeBoneTransform = PreviewSkelMeshComp->GetBoneTransform(PreviewSkelMeshComp->GetBoneIndex(ActivePreviewNode->RelativeSpaceBone.BoneName));
+						BodyTransform = BodyTransform * RelativeBoneTransform;
+					}
+
+					for (const FAnimPhysShape& Shape : Body.Shapes)
+					{
+						for (const FIntVector& Triangle : Shape.Triangles)
+						{
+							for (int32 Idx = 0; Idx < 3; ++Idx)
+							{
+								int32 Next = (Idx + 1) % 3;
+
+								FVector FirstVertPosition = BodyTransform.TransformPosition(Shape.Vertices[Triangle[Idx]]);
+								FVector SecondVertPosition = BodyTransform.TransformPosition(Shape.Vertices[Triangle[Next]]);
+
+								PDI->DrawLine(FirstVertPosition, SecondVertPosition, AnimDynamicsNodeConstants::ActiveBodyDrawColor, SDPG_Foreground, AnimDynamicsNodeConstants::ShapeLineWidth);
+							}
+						}
+					}
+
+					const int32 BoneIndex = PreviewSkelMeshComp->GetBoneIndex(ActivePreviewNode->BoundBone.BoneName);
+					if (BoneIndex != INDEX_NONE)
+					{
+						FTransform BodyJointTransform = PreviewSkelMeshComp->GetBoneTransform(BoneIndex);
+						FTransform ShapeOriginalTransform = BodyJointTransform;
+
+						// Draw pin location
+						const FVector LocalPinOffset = - BodyTransform.Rotator().RotateVector(ActivePreviewNode->GetBodyLocalJointOffset(BodyIndex)); // Position of joint relative to physics body in world space.
+						PDI->DrawLine(Body.Pose.Position, Body.Pose.Position + LocalPinOffset, FLinearColor::Green, SDPG_Foreground, AnimDynamicsNodeConstants::ShapeLineWidth);
+
+						// Draw basis at body location
+						FVector Origin = BodyTransform.GetTranslation();
+						FVector XAxis(1.0f, 0.0f, 0.0f);
+						FVector YAxis(0.0f, 1.0f, 0.0f);
+						FVector ZAxis(0.0f, 0.0f, 1.0f);
+
+						XAxis = BodyTransform.TransformVector(XAxis);
+						YAxis = BodyTransform.TransformVector(YAxis);
+						ZAxis = BodyTransform.TransformVector(ZAxis);
+
+						PDI->DrawLine(Origin, Origin + XAxis * AnimDynamicsNodeConstants::TransformBasisScale, FLinearColor::Red, SDPG_Foreground, AnimDynamicsNodeConstants::TransformLineWidth);
+						PDI->DrawLine(Origin, Origin + YAxis * AnimDynamicsNodeConstants::TransformBasisScale, FLinearColor::Green, SDPG_Foreground, AnimDynamicsNodeConstants::TransformLineWidth);
+						PDI->DrawLine(Origin, Origin + ZAxis * AnimDynamicsNodeConstants::TransformBasisScale, FLinearColor::Blue, SDPG_Foreground, AnimDynamicsNodeConstants::TransformLineWidth);
+					}
+				}
+			}
+
+			uint32 PhysicsBodyIndex = 0;
+
+			for (const FAnimPhysBodyDefinition& BodyDefinition : ActivePreviewNode->PhysicsBodyDefinitions)
+			{
+				const int32 BoneIndex = PreviewSkelMeshComp->GetBoneIndex(BodyDefinition.BoundBone.BoneName);
+				if (BoneIndex != INDEX_NONE)
+				{
+					// World space transform
+					const FTransform BoneTransform = PreviewSkelMeshComp->GetBoneTransform(BoneIndex);
+					FTransform ShapeTransform = BoneTransform;
+					ShapeTransform.SetTranslation(ShapeTransform.GetTranslation() + BoneTransform.GetRotation().RotateVector(BodyDefinition.LocalJointOffset)); // Transform of rendered shape in world space.
+
+					const FAnimPhysShape EditPreviewShape = FAnimPhysShape::MakeBox(BodyDefinition.BoxExtents);
+
+					{
+						const FAnimDynamicsViewportObjectReference ViewportObjectRef(EditorAnimDynamicsNode->GetUniqueID(), FAnimDynamicsViewportObjectType::BoxExtents, PhysicsBodyIndex);
+						const bool bIsSelected = SelectedViewportObjects.Contains(ViewportObjectRef);
+						const float LineWidth = (bIsSelected) ? AnimDynamicsNodeConstants::SelectedShapeLineWidth : AnimDynamicsNodeConstants::ShapeLineWidth;
+
+						// Draw Physics Body
+						PDI->SetHitProxy(new HAnimDynamicsEditModeHitProxy(ViewportObjectRef));
+						DrawShape(PDI, EditPreviewShape, ShapeTransform, AnimDynamicsNodeConstants::ShapeDrawColor, LineWidth);
+						PDI->SetHitProxy(nullptr);
+					}
+
+					DrawBasis(PDI, ShapeTransform);
+
+					// Draw line connecting rendered shape to its associated bone.
+					PDI->DrawLine(ShapeTransform.GetTranslation(), BoneTransform.GetTranslation(), FLinearColor::Green, SDPG_Foreground, AnimDynamicsNodeConstants::ShapeLineWidth);
+
+					if (EditorAnimDynamicsNode->bShowLinearLimits)
+					{
+						DrawLinearLimits(PDI, ShapeTransform, BodyDefinition.ConstraintSetup);
+					}
+
+					if (EditorAnimDynamicsNode->bShowAngularLimits)
+					{
+						DrawAngularLimits(PDI, ShapeTransform, BodyDefinition.ConstraintSetup);
+					}
+
+					if (EditorAnimDynamicsNode->bShowCollisionSpheres && BodyDefinition.CollisionType != AnimPhysCollisionType::CoM)
+					{
+						const FAnimDynamicsViewportObjectReference ViewportObjectRef(EditorAnimDynamicsNode->GetUniqueID(), FAnimDynamicsViewportObjectType::SphericalColisionVolume, PhysicsBodyIndex);
+						const bool bIsSelected = SelectedViewportObjects.Contains(ViewportObjectRef);
+						const float LineWidth = (bIsSelected) ? AnimDynamicsNodeConstants::SelectedShapeLineWidth : AnimDynamicsNodeConstants::ShapeLineWidth;
+
+						// Draw Collision Sphere.
+						PDI->SetHitProxy(new HAnimDynamicsEditModeHitProxy(ViewportObjectRef));
+						DrawWireSphere(PDI, ShapeTransform, FLinearColor(FColor::Cyan), BodyDefinition.SphereCollisionRadius, 24, SDPG_Foreground, LineWidth);
+						PDI->SetHitProxy(nullptr);
+					}
+				}
+
+				++PhysicsBodyIndex;
+			}
+
+			const float LimitPlaneSize = AnimDynamicsNodeConstants::LimitPlaneDrawSize;
+			const FLinearColor LimitEdgeColor = AnimDynamicsNodeConstants::LimitLineDrawColor;
+
+
+			// Draw the planar limits.
+			if (EditorAnimDynamicsNode->bShowPlanarLimit && ActivePreviewNode->PlanarLimits.Num() > 0)
+			{
+				uint32 LimitIndex = 0;
+				for (const FAnimPhysPlanarLimit& PlanarLimit : ActivePreviewNode->PlanarLimits)
+				{
+					FTransform LimitPlaneTransform = PlanarLimit.PlaneTransform;
+					const int32 LimitDrivingBoneIdx = PreviewSkelMeshComp->GetBoneIndex(PlanarLimit.DrivingBone.BoneName);
+
+					if (LimitDrivingBoneIdx != INDEX_NONE)
+					{
+						LimitPlaneTransform *= PreviewSkelMeshComp->GetComponentSpaceTransforms()[LimitDrivingBoneIdx];
+					}
+
+					const UE::Math::TMatrix< float > PlaneTransform = LimitPlaneTransform.ToMatrixNoScale();
+
+					const FAnimDynamicsViewportObjectReference ViewportObjectRef(EditorAnimDynamicsNode->GetUniqueID(), FAnimDynamicsViewportObjectType::PlaneLimit, LimitIndex);
+					const bool bIsSelected = SelectedViewportObjects.Contains(ViewportObjectRef);
+					const float LineWidth = (bIsSelected) ? AnimDynamicsNodeConstants::SelectedShapeLineWidth : AnimDynamicsNodeConstants::ShapeLineWidth;
+
+					PDI->SetHitProxy(new HAnimDynamicsEditModeHitProxy(ViewportObjectRef));
+					DrawPlane10x10(PDI, PlaneTransform, LimitPlaneSize, FVector2D(0.0f, 0.0f), FVector2D(1.0f, 1.0f), GEngine->ConstraintLimitMaterialPrismatic->GetRenderProxy(), SDPG_World);
+
+					const FVector PlaneVertexPositionA(PlaneTransform.TransformPosition(FVector(-1.0, -1.0, 0.0) * LimitPlaneSize));
+					const FVector PlaneVertexPositionB(PlaneTransform.TransformPosition(FVector(1.0, -1.0, 0.0) * LimitPlaneSize));
+					const FVector PlaneVertexPositionC(PlaneTransform.TransformPosition(FVector(1.0, 1.0, 0.0) * LimitPlaneSize));
+					const FVector PlaneVertexPositionD(PlaneTransform.TransformPosition(FVector(-1.0, 1.0, 0.0) * LimitPlaneSize));
+
+					PDI->DrawLine(PlaneVertexPositionA, PlaneVertexPositionB, LimitEdgeColor, SDPG_Foreground, LineWidth);
+					PDI->DrawLine(PlaneVertexPositionB, PlaneVertexPositionC, LimitEdgeColor, SDPG_Foreground, LineWidth);
+					PDI->DrawLine(PlaneVertexPositionC, PlaneVertexPositionD, LimitEdgeColor, SDPG_Foreground, LineWidth);
+					PDI->DrawLine(PlaneVertexPositionD, PlaneVertexPositionA, LimitEdgeColor, SDPG_Foreground, LineWidth);
+					PDI->DrawLine(PlaneVertexPositionA, PlaneVertexPositionC, LimitEdgeColor, SDPG_Foreground, LineWidth);
+					PDI->DrawLine(PlaneVertexPositionD, PlaneVertexPositionB, LimitEdgeColor, SDPG_Foreground, LineWidth);
+
+					DrawDirectionalArrow(PDI, FRotationMatrix(FRotator(90.0f, 0.0f, 0.0f)) * LimitPlaneTransform.ToMatrixNoScale(), LimitEdgeColor, 50.0f, 20.0f, SDPG_Foreground, LineWidth);
+					PDI->SetHitProxy(nullptr);
+
+					++LimitIndex;
+				}
+			}
+
+			// Draw Spherical Limits.
+			if (EditorAnimDynamicsNode->bShowSphericalLimit && ActivePreviewNode->SphericalLimits.Num() > 0)
+			{
+				uint32 LimitIndex = 0;
+
+				for (const FAnimPhysSphericalLimit& SphericalLimit : ActivePreviewNode->SphericalLimits)
+				{
+					FTransform SphereTransform = FTransform::Identity;
+					SphereTransform.SetTranslation(SphericalLimit.SphereLocalOffset);
+
+					const int32 DrivingBoneIdx = PreviewSkelMeshComp->GetBoneIndex(SphericalLimit.DrivingBone.BoneName);
+
+					if (DrivingBoneIdx != INDEX_NONE)
+					{
+						SphereTransform *= PreviewSkelMeshComp->GetComponentSpaceTransforms()[DrivingBoneIdx];
+					}
+
+					const FAnimDynamicsViewportObjectReference ViewportObjectRef(EditorAnimDynamicsNode->GetUniqueID(), FAnimDynamicsViewportObjectType::SphericalLimit, LimitIndex);
+					const bool bIsSelected = SelectedViewportObjects.Contains(ViewportObjectRef);
+					const float LineWidth = (bIsSelected) ? AnimDynamicsNodeConstants::SelectedShapeLineWidth : AnimDynamicsNodeConstants::ShapeLineWidth;
+
+					PDI->SetHitProxy(new HAnimDynamicsEditModeHitProxy(ViewportObjectRef));
+					DrawSphere(PDI, SphereTransform.GetLocation(), FRotator::ZeroRotator, FVector(SphericalLimit.LimitRadius), 24, 6, GEngine->ConstraintLimitMaterialPrismatic->GetRenderProxy(), SDPG_World);
+					DrawWireSphere(PDI, SphereTransform, LimitEdgeColor, SphericalLimit.LimitRadius, 24, SDPG_World, LineWidth);
+					PDI->SetHitProxy(nullptr);
+
+					++LimitIndex;
+				}
+			}
+		}
+	}
+}
+
+bool FAnimDynamicsEditMode::HandleClick(FEditorViewportClient* InViewportClient, HHitProxy* HitProxy, const FViewportClick& Click)
+{
+	const FViewport* const Viewport = InViewportClient->Viewport;
+
+	const bool bIsCtrlKeyDown = Viewport->KeyState(EKeys::LeftControl) || Viewport->KeyState(EKeys::RightControl);
+	const bool bIsShiftKeyDown = Viewport->KeyState(EKeys::LeftShift) || Viewport->KeyState(EKeys::RightShift);
+	const bool bIsAltKeyDown = Viewport->KeyState(EKeys::LeftAlt) || Viewport->KeyState(EKeys::RightAlt);
+	bool bHandled = false;
+
+	const bool bModifySelection = bIsCtrlKeyDown | bIsShiftKeyDown;
+
+	const HAnimDynamicsEditModeHitProxy* const AnimDynamicsProxy = HitProxyCast<HAnimDynamicsEditModeHitProxy>(HitProxy);
+
+	if (AnimDynamicsProxy)
+	{
+		const FAnimDynamicsViewportObjectReference HitObjectRef = AnimDynamicsProxy->ViewportObjectRef;
+
+		if (!bModifySelection)
+		{
+			SelectedViewportObjects.Reset();
+		}
+
+		if (bModifySelection && SelectedViewportObjects.Contains(HitObjectRef))
+		{
+			SelectedViewportObjects.Remove(HitObjectRef);
+		}
+		else
+		{
+			SelectedViewportObjects.Emplace(HitObjectRef);
+		}
+
+		bHandled = true;
+	}
+
+	if (!bHandled)
+	{
+		bHandled = FEdMode::HandleClick(InViewportClient, HitProxy, Click);
+	}
+
+	if (!bHandled && !bModifySelection)
+	{
+		SelectedViewportObjects.Reset(); // Clear selection if Click on empty space.
+		bHandled = true;
+	}
+
+	return bHandled;
+}
+
+ECoordSystem FAnimDynamicsEditMode::GetWidgetCoordinateSystem() const
+{
+	return COORD_Local;
+}
+
+FVector FAnimDynamicsEditMode::GetWidgetLocation() const
+{
+	return GetActiveViewportObjectTransform().GetTranslation();
+}
+
+bool FAnimDynamicsEditMode::GetCustomDrawingCoordinateSystem(FMatrix& InMatrix, void* InData)
+{
+	InMatrix = GetActiveViewportObjectTransform().ToMatrixNoScale().RemoveTranslation();
+	return true;
+}
+
+
+const bool FAnimDynamicsEditMode::IsValidWidgetMode(UE::Widget::EWidgetMode InWidgetMode) const
+{
+	return (InWidgetMode == UE::Widget::WM_Scale) || (InWidgetMode == UE::Widget::WM_Translate);
+}
+
+UE::Widget::EWidgetMode FAnimDynamicsEditMode::GetNextWidgetMode(UE::Widget::EWidgetMode InWidgetMode) const
+{
+	UE::Widget::EWidgetMode NextMode = UE::Widget::WM_Translate;
+
+	if (InWidgetMode == UE::Widget::WM_Translate)
+	{
+		NextMode = UE::Widget::WM_Scale;
+	}
+
+	check(IsValidWidgetMode(NextMode));
+
+	return NextMode;
+}
+
+UE::Widget::EWidgetMode FAnimDynamicsEditMode::FindValidWidgetMode(UE::Widget::EWidgetMode InWidgetMode) const
+{
+	UE::Widget::EWidgetMode ValidMode = InWidgetMode;
+
+	if (!IsValidWidgetMode(ValidMode))
+	{
+		ValidMode = GetNextWidgetMode(ValidMode);
+	}
+
+	check(IsValidWidgetMode(ValidMode));
+
+	return ValidMode;
+}
+
+UE::Widget::EWidgetMode FAnimDynamicsEditMode::GetWidgetMode() const
+{
+	const USkeletalMeshComponent* const SkelComp = GetAnimPreviewScene().GetPreviewMeshComponent();
+	const UAnimGraphNode_AnimDynamics* const EditorAnimDynamicsNode = CastChecked<UAnimGraphNode_AnimDynamics>(GetActiveWidgetAnimNode());
+
+	if (SkelComp && EditorAnimDynamicsNode)
+	{
+		const int32 BoneIndex = SkelComp->GetBoneIndex(EditorAnimDynamicsNode->Node.BoundBone.BoneName);
+		if (BoneIndex != INDEX_NONE)
+		{
+			CurWidgetMode = FindValidWidgetMode(CurWidgetMode);
+			return CurWidgetMode;
+		}
+	}
+
+	return UE::Widget::WM_None;
+}
+
+UE::Widget::EWidgetMode FAnimDynamicsEditMode::ChangeToNextWidgetMode(UE::Widget::EWidgetMode InCurWidgetMode)
+{
+	UE::Widget::EWidgetMode NextWidgetMode = GetNextWidgetMode(InCurWidgetMode);
+	CurWidgetMode = FindValidWidgetMode(NextWidgetMode);
+
+	return CurWidgetMode;
+}
+
+bool FAnimDynamicsEditMode::SetWidgetMode(UE::Widget::EWidgetMode InWidgetMode)
+{
+	// if InWidgetMode is available 
+	if (IsValidWidgetMode(InWidgetMode))
+	{
+		CurWidgetMode = InWidgetMode;
+		return true;
+	}
+
+	return false;
+}
+
+FName FAnimDynamicsEditMode::GetSelectedBone() const
+{
+	FName SelectedBoneName;
+
+	if (SelectedViewportObjects.Num() > 0)
+	{
+		const uint32 Index = SelectedViewportObjects.Num() - 1;
+		const FAnimDynamicsViewportObjectReference& SelectedObjectRef = SelectedViewportObjects[Index];
+		const UAnimGraphNode_AnimDynamics* const EditorAnimDynamicsNode = FindSelectedEditorAnimNode(SelectedObjectRef.EditorNodeUniqueId);
+		const FAnimNode_AnimDynamics* const ActivePreviewNode = EditorAnimDynamicsNode->GetPreviewDynamicsNode();
+
+		{
+			if ((SelectedObjectRef.Type == FAnimDynamicsViewportObjectType::BoxExtents) || (SelectedObjectRef.Type == FAnimDynamicsViewportObjectType::SphericalColisionVolume))
+			{
+				if (ActivePreviewNode->PhysicsBodyDefinitions.IsValidIndex(SelectedObjectRef.Index))
+				{
+					SelectedBoneName = ActivePreviewNode->PhysicsBodyDefinitions[SelectedObjectRef.Index].BoundBone.BoneName;
+				}
+			}
+			else if (SelectedObjectRef.Type == FAnimDynamicsViewportObjectType::PlaneLimit)
+			{
+				if (ActivePreviewNode->PlanarLimits.IsValidIndex(SelectedObjectRef.Index))
+				{
+					SelectedBoneName = ActivePreviewNode->PlanarLimits[SelectedObjectRef.Index].DrivingBone.BoneName;
+				}
+			}
+			else if (SelectedObjectRef.Type == FAnimDynamicsViewportObjectType::SphericalLimit)
+			{
+				if (ActivePreviewNode->SphericalLimits.IsValidIndex(SelectedObjectRef.Index))
+				{
+					SelectedBoneName = ActivePreviewNode->SphericalLimits[SelectedObjectRef.Index].DrivingBone.BoneName;
+				}
+			}
+		}
+	}
+
+	return SelectedBoneName;
+}
+
+const UAnimGraphNode_AnimDynamics* const FAnimDynamicsEditMode::FindSelectedEditorAnimNode(const int32 InEditorNodeId) const
+{
+	const FAnimNodeEditMode::EditorRuntimeNodePair* const FoundNodePair = AnimNodes.FindByPredicate([InEditorNodeId](const FAnimNodeEditMode::EditorRuntimeNodePair& Element) { return Element.EditorAnimNode && Element.EditorAnimNode->GetUniqueID() == InEditorNodeId; });
+
+	if (FoundNodePair)
+	{
+		return CastChecked< UAnimGraphNode_AnimDynamics >(FoundNodePair->EditorAnimNode);
+	}
+
+	return nullptr;
+}
+
+UAnimGraphNode_AnimDynamics* const FAnimDynamicsEditMode::FindSelectedEditorAnimNode(const int32 InEditorNodeId)
+{
+	const FAnimNodeEditMode::EditorRuntimeNodePair* const FoundNodePair = AnimNodes.FindByPredicate([InEditorNodeId](const FAnimNodeEditMode::EditorRuntimeNodePair& Element) { return Element.EditorAnimNode && Element.EditorAnimNode->GetUniqueID() == InEditorNodeId; });
+
+	if (FoundNodePair)
+	{
+		return CastChecked< UAnimGraphNode_AnimDynamics >(FoundNodePair->EditorAnimNode);
+	}
+
+	return nullptr;
+}
+
+const FTransform FAnimDynamicsEditMode::GetActiveViewportObjectTransform() const
+{
+	return GetViewportObjectTransform(GetActiveViewportObject());
+}
+
+const FTransform FAnimDynamicsEditMode::GetViewportObjectTransform(const FAnimDynamicsViewportObjectReference* const SelectedObjectRef) const
+{
+	FTransform Transform;
+	Transform.SetIdentity();
+
+	if (SelectedObjectRef)
+	{
+		Transform = GetViewportObjectLocalSpaceTransform(SelectedObjectRef);
+		const UAnimGraphNode_AnimDynamics* const EditorAnimDynamicsNode = FindSelectedEditorAnimNode(SelectedObjectRef->EditorNodeUniqueId);
+
+		FName BoneName;
+
+		if (EditorAnimDynamicsNode)
+		{
+			const FAnimNode_AnimDynamics* const ActivePreviewNode = EditorAnimDynamicsNode->GetPreviewDynamicsNode();
+
+			if (ActivePreviewNode)
+			{
+				if ((SelectedObjectRef->Type == FAnimDynamicsViewportObjectType::BoxExtents) || (SelectedObjectRef->Type == FAnimDynamicsViewportObjectType::SphericalColisionVolume))
+				{
+					if (ActivePreviewNode->PhysicsBodyDefinitions.IsValidIndex(SelectedObjectRef->Index))
+					{
+						Transform.SetTranslation(Transform.GetTranslation() + Transform.GetRotation().RotateVector(ActivePreviewNode->PhysicsBodyDefinitions[SelectedObjectRef->Index].LocalJointOffset));
+					}
+				}
+				else if (SelectedObjectRef->Type == FAnimDynamicsViewportObjectType::PlaneLimit)
+				{
+					if (ActivePreviewNode->PlanarLimits.IsValidIndex(SelectedObjectRef->Index))
+					{
+						Transform = ActivePreviewNode->PlanarLimits[SelectedObjectRef->Index].PlaneTransform * Transform;
+					}
+				}
+				else if (SelectedObjectRef->Type == FAnimDynamicsViewportObjectType::SphericalLimit)
+				{
+					if (ActivePreviewNode->SphericalLimits.IsValidIndex(SelectedObjectRef->Index))
+					{
+						Transform.SetTranslation(Transform.GetTranslation() + Transform.GetRotation().RotateVector(ActivePreviewNode->SphericalLimits[SelectedObjectRef->Index].SphereLocalOffset));
+					}
+				}
+			}
+		}
+	}
+
+
+	return Transform;
+}
+
+const FTransform FAnimDynamicsEditMode::GetViewportObjectLocalSpaceTransform(const FAnimDynamicsViewportObjectReference* const SelectedObjectRef) const
+{
+	FTransform Transform;
+	Transform.SetIdentity();
+
+	const USkeletalMeshComponent* const SkelComp = GetAnimPreviewScene().GetPreviewMeshComponent();
+
+	if (SelectedObjectRef && SkelComp)
+	{
+		const FName BoneName = GetSelectedBone();
+		const int32 BoneIndex = SkelComp->GetBoneIndex(BoneName);
+
+		if (BoneIndex != INDEX_NONE)
+		{
+			Transform = SkelComp->GetComponentSpaceTransforms()[BoneIndex];
+		}
+	}
+
+	return Transform;
+}
+
+void FAnimDynamicsEditMode::DoTranslation(FVector& InTranslation)
+{
+	const USkeletalMeshComponent* const SkelComp = GetAnimPreviewScene().GetPreviewMeshComponent();
+
+	if (SkelComp)
+	{
+		// Find widget translation in the space of the selected objects associated bone.
+		const FTransform LocalSpaceTransform = GetViewportObjectLocalSpaceTransform(GetActiveViewportObject());
+		const FVector LocalSpaceTranslation = LocalSpaceTransform.GetRotation().UnrotateVector(InTranslation);
+
+		// Apply local space widget translation to all selected nodes. This means the collision volumes will each move in the same direction 
+		// in their own space, potentially different directions in world space.
+		for (const FAnimDynamicsViewportObjectReference& SelectedObjectRef : SelectedViewportObjects)
+		{
+			UAnimGraphNode_AnimDynamics* const EditorAnimDynamicsNode = FindSelectedEditorAnimNode(SelectedObjectRef.EditorNodeUniqueId);
+			FAnimNode_AnimDynamics* const RuntimeAnimDynamicsNode = EditorAnimDynamicsNode->GetPreviewDynamicsNode();
+
+			check(EditorAnimDynamicsNode && RuntimeAnimDynamicsNode);
+
+			if (EditorAnimDynamicsNode && RuntimeAnimDynamicsNode)
+			{
+				if ((SelectedObjectRef.Type == FAnimDynamicsViewportObjectType::BoxExtents) || (SelectedObjectRef.Type == FAnimDynamicsViewportObjectType::SphericalColisionVolume))
+				{
+					if (RuntimeAnimDynamicsNode->PhysicsBodyDefinitions.IsValidIndex(SelectedObjectRef.Index))
+					{
+						FVector& PhysicsBodyJointOffset = RuntimeAnimDynamicsNode->PhysicsBodyDefinitions[SelectedObjectRef.Index].LocalJointOffset;
+
+						PhysicsBodyJointOffset += LocalSpaceTranslation;
+						EditorAnimDynamicsNode->Node.PhysicsBodyDefinitions[SelectedObjectRef.Index].LocalJointOffset = PhysicsBodyJointOffset;
+					}
+				}
+				else if (SelectedObjectRef.Type == FAnimDynamicsViewportObjectType::PlaneLimit)
+				{
+					if (RuntimeAnimDynamicsNode->PlanarLimits.IsValidIndex(SelectedObjectRef.Index))
+					{
+						FTransform& LimitTransform = RuntimeAnimDynamicsNode->PlanarLimits[SelectedObjectRef.Index].PlaneTransform;
+						LimitTransform.SetTranslation(LimitTransform.GetTranslation() + LocalSpaceTranslation);
+						EditorAnimDynamicsNode->Node.PlanarLimits[SelectedObjectRef.Index].PlaneTransform = LimitTransform;
+					}
+				}
+				else if (SelectedObjectRef.Type == FAnimDynamicsViewportObjectType::SphericalLimit)
+				{
+					if (RuntimeAnimDynamicsNode->SphericalLimits.IsValidIndex(SelectedObjectRef.Index))
+					{
+
+						FVector& LimitOffset = RuntimeAnimDynamicsNode->SphericalLimits[SelectedObjectRef.Index].SphereLocalOffset;
+						LimitOffset += LocalSpaceTranslation;
+						EditorAnimDynamicsNode->Node.SphericalLimits[SelectedObjectRef.Index].SphereLocalOffset = LimitOffset;
+					}
+				}
+			}
+		}
+	}
+
+	bIsInteractingWithWidget = true;
+}
+
+void FAnimDynamicsEditMode::DoRotation(FRotator& InRotation)
+{
+	const USkeletalMeshComponent* const SkelComp = GetAnimPreviewScene().GetPreviewMeshComponent();
+
+	if (SkelComp)
+	{
+		// Transform rotation to the active object's local space.
+		FQuat LocalSpaceRotation(FQuat::Identity);
+
+		{
+			FVector RotAxis;
+			float RotAngle;
+			InRotation.Quaternion().ToAxisAndAngle(RotAxis, RotAngle);
+			const FTransform SelectedBoneTM = GetViewportObjectLocalSpaceTransform(GetActiveViewportObject());
+			RotAxis = SelectedBoneTM.Inverse().TransformVector(RotAxis);
+			FQuat DeltaQuat(RotAxis, RotAngle);
+			DeltaQuat.Normalize();
+			LocalSpaceRotation = DeltaQuat;
+		}
+
+		// Apply local rotation to all selected objects.
+		for (const FAnimDynamicsViewportObjectReference& SelectedObjectRef : SelectedViewportObjects)
+		{
+			UAnimGraphNode_AnimDynamics* const EditorAnimDynamicsNode = FindSelectedEditorAnimNode(SelectedObjectRef.EditorNodeUniqueId);
+			FAnimNode_AnimDynamics* const RuntimeAnimDynamicsNode = EditorAnimDynamicsNode->GetPreviewDynamicsNode();
+
+			check(EditorAnimDynamicsNode && RuntimeAnimDynamicsNode);
+			if (EditorAnimDynamicsNode && RuntimeAnimDynamicsNode)
+			{
+				if (SelectedObjectRef.Type == FAnimDynamicsViewportObjectType::PlaneLimit)
+				{
+					if (RuntimeAnimDynamicsNode->PlanarLimits.IsValidIndex(SelectedObjectRef.Index))
+					{
+						FTransform& LimitTransform = RuntimeAnimDynamicsNode->PlanarLimits[SelectedObjectRef.Index].PlaneTransform;
+						LimitTransform.SetRotation(LocalSpaceRotation * LimitTransform.GetRotation());
+						EditorAnimDynamicsNode->Node.PlanarLimits[SelectedObjectRef.Index].PlaneTransform = LimitTransform;
+					}
+				}
+			}
+		}
+	}
+}
+
+void FAnimDynamicsEditMode::DoScale(FVector& InScale)
+{
+	const float Scale = InScale.X + InScale.Y + InScale.Z;
+
+	for (const FAnimDynamicsViewportObjectReference& SelectedObjectRef : SelectedViewportObjects)
+	{
+		UAnimGraphNode_AnimDynamics* const EditorAnimDynamicsNode = FindSelectedEditorAnimNode(SelectedObjectRef.EditorNodeUniqueId);
+		FAnimNode_AnimDynamics* const RuntimeAnimDynamicsNode = EditorAnimDynamicsNode->GetPreviewDynamicsNode();
+
+		check(EditorAnimDynamicsNode && RuntimeAnimDynamicsNode);
+
+		if (EditorAnimDynamicsNode && RuntimeAnimDynamicsNode)
+		{
+			if (SelectedObjectRef.Type == FAnimDynamicsViewportObjectType::BoxExtents)
+			{
+				if (RuntimeAnimDynamicsNode->PhysicsBodyDefinitions.IsValidIndex(SelectedObjectRef.Index))
+				{
+					FVector& PhysicsBodyBoxExtents = RuntimeAnimDynamicsNode->PhysicsBodyDefinitions[SelectedObjectRef.Index].BoxExtents;
+					PhysicsBodyBoxExtents += InScale;
+					EditorAnimDynamicsNode->Node.PhysicsBodyDefinitions[SelectedObjectRef.Index].BoxExtents = PhysicsBodyBoxExtents;
+				}
+			}
+			else if (SelectedObjectRef.Type == FAnimDynamicsViewportObjectType::SphericalColisionVolume)
+			{
+				if (RuntimeAnimDynamicsNode->PhysicsBodyDefinitions.IsValidIndex(SelectedObjectRef.Index))
+				{
+					float& PhysicsBodySphereCollisionRadius = RuntimeAnimDynamicsNode->PhysicsBodyDefinitions[SelectedObjectRef.Index].SphereCollisionRadius;
+
+					PhysicsBodySphereCollisionRadius += Scale;
+					EditorAnimDynamicsNode->Node.PhysicsBodyDefinitions[SelectedObjectRef.Index].SphereCollisionRadius = PhysicsBodySphereCollisionRadius;
+				}
+			}
+			else if (SelectedObjectRef.Type == FAnimDynamicsViewportObjectType::SphericalLimit)
+			{
+				if (RuntimeAnimDynamicsNode->SphericalLimits.IsValidIndex(SelectedObjectRef.Index))
+				{
+					const float LimitRadius = RuntimeAnimDynamicsNode->SphericalLimits[SelectedObjectRef.Index].LimitRadius + Scale;
+
+					RuntimeAnimDynamicsNode->SphericalLimits[SelectedObjectRef.Index].LimitRadius = LimitRadius;
+					EditorAnimDynamicsNode->Node.SphericalLimits[SelectedObjectRef.Index].LimitRadius = LimitRadius;
+				}
+			}
+		}
+	}
+
+	bIsInteractingWithWidget = true;
+}
+
+void FAnimDynamicsEditMode::Tick(FEditorViewportClient* ViewportClient, float DeltaTime)
+{
+	FAnimNodeEditMode::Tick(ViewportClient, DeltaTime);
+
+	// Keep this flag set until the mouse is released after moving the translation widget. DoTranslation / Scale fns are only called when 
+	// the widget actually moves but IsManipulatingWidget() rtns true whenever the input button is pressed. This logic prevents the physics 
+	// simulation from stating/stopping unexpectedly when the widget is being manipulated but hasn't moved and avoids stopping the simulation 
+	// every time a mouse button is pressed over the viewport.
+	bIsInteractingWithWidget = IsManipulatingWidget() && bIsInteractingWithWidget;
+
+	USkeletalMeshComponent* const PreviewSkelMeshComp = GetAnimPreviewScene().GetPreviewMeshComponent();
+
+	for (const EditorRuntimeNodePair& CurrentNodePair : AnimNodes)
+	{
+		const UAnimGraphNode_AnimDynamics* const EditorAnimDynamicsNode = CastChecked<UAnimGraphNode_AnimDynamics>(CurrentNodePair.EditorAnimNode);
+
+		if (EditorAnimDynamicsNode)
+		{
+			FAnimNode_AnimDynamics* const ActivePreviewNode = EditorAnimDynamicsNode->GetPreviewDynamicsNode();
+
+			if (ActivePreviewNode)
+			{
+				// Disable physics simulation when physics parameters are being edited via a widget, otherwise editing becomes 
+				// very difficult as the objects are moving and the simulation is liable to become unstable.
+				if (bIsInteractingWithWidget)
+				{
+					ActivePreviewNode->bDoPhysicsUpdateInEditor = false;
+				}
+				else if (!ActivePreviewNode->bDoPhysicsUpdateInEditor)
+				{
+					ActivePreviewNode->RequestInitialise(ETeleportType::ResetPhysics);
+					ActivePreviewNode->bDoPhysicsUpdateInEditor = true;
+				}
+			}
+
+			EditorAnimDynamicsNode->LastPreviewComponent = PreviewSkelMeshComp;
+		}
+	}
+}
+
+bool FAnimDynamicsEditMode::ShouldDrawWidget() const
+{
+	return SelectedViewportObjects.Num() > 0; // Only draw a widget if at least one viewport object is selected.
+}
+
+const FAnimDynamicsViewportObjectReference* const FAnimDynamicsEditMode::GetActiveViewportObject() const
+{
+	if (SelectedViewportObjects.Num() > 0)
+	{
+		return &SelectedViewportObjects[SelectedViewportObjects.Num() - 1];
+	}
+
+	return nullptr;
+}
+
diff --git a/Engine/Source/Editor/AnimGraph/Private/EditModes/AnimDynamicsEditMode.h b/Engine/Source/Editor/AnimGraph/Private/EditModes/AnimDynamicsEditMode.h
new file mode 100644
--- /dev/null
+++ b/Engine/Source/Editor/AnimGraph/Private/EditModes/AnimDynamicsEditMode.h
@@ -0,0 +1,68 @@
+// Copyright Epic Games, Inc. All Rights Reserved.
+
+#pragma once
+
+#include "CoreMinimal.h"
+#include "UnrealWidgetFwd.h"
+#include "AnimNodeEditMode.h"
+
+enum class FAnimDynamicsViewportObjectType { PlaneLimit, SphericalLimit, SphericalColisionVolume, BoxExtents };
+
+// FAnimDynamicsViewportObjectReference
+//
+// Used to identify a viewport editable structure (i.e. a physics object joint offset) in an AnimDynamics node.
+struct FAnimDynamicsViewportObjectReference
+{
+	FAnimDynamicsViewportObjectReference(const uint32 InEditorNodeUniqueId, const FAnimDynamicsViewportObjectType InType, const uint32 Index);
+	
+	uint32 EditorNodeUniqueId;
+	FAnimDynamicsViewportObjectType Type;
+	uint32 Index;
+};
+
+const bool operator==(const FAnimDynamicsViewportObjectReference& Lhs, const FAnimDynamicsViewportObjectReference& Rhs);
+
+class FAnimDynamicsEditMode : public FAnimNodeEditMode
+{
+public:
+	FAnimDynamicsEditMode();
+
+	/** IAnimNodeEditMode interface */
+	virtual void Render(const FSceneView* View, FViewport* Viewport, FPrimitiveDrawInterface* PDI) override;
+	virtual bool HandleClick(FEditorViewportClient* InViewportClient, HHitProxy* HitProxy, const FViewportClick& Click) override;
+	virtual void EnterMode(class UAnimGraphNode_Base* InEditorNode, struct FAnimNode_Base* InRuntimeNode) override;
+	virtual void ExitMode() override;
+	virtual ECoordSystem GetWidgetCoordinateSystem() const override;
+	virtual FVector GetWidgetLocation() const override;
+	virtual bool GetCustomDrawingCoordinateSystem(FMatrix& InMatrix, void* InData) override;
+	virtual UE::Widget::EWidgetMode GetWidgetMode() const override;
+	virtual UE::Widget::EWidgetMode ChangeToNextWidgetMode(UE::Widget::EWidgetMode InCurWidgetMode) override;
+	virtual bool SetWidgetMode(UE::Widget::EWidgetMode InWidgetMode) override;
+	virtual FName GetSelectedBone() const override;
+	virtual void DoTranslation(FVector& InTranslation) override;
+	virtual void DoRotation(FRotator& InRotation) override;
+	virtual void DoScale(FVector& InScale) override;
+	virtual void Tick(FEditorViewportClient* ViewportClient, float DeltaTime) override;
+	virtual bool ShouldDrawWidget() const override;
+
+private:
+	UE::Widget::EWidgetMode FindValidWidgetMode(UE::Widget::EWidgetMode InWidgetMode) const;
+	const class UAnimGraphNode_AnimDynamics* const FindSelectedEditorAnimNode(const int32 InEditorNodeId) const;
+	class UAnimGraphNode_AnimDynamics* const FindSelectedEditorAnimNode(const int32 InEditorNodeId);
+	const bool IsValidWidgetMode(UE::Widget::EWidgetMode InWidgetMode) const;
+	UE::Widget::EWidgetMode GetNextWidgetMode(UE::Widget::EWidgetMode InWidgetMode) const;
+	
+	const FTransform GetActiveViewportObjectTransform() const;
+	const FTransform GetViewportObjectTransform(const FAnimDynamicsViewportObjectReference* const SelectedObjectRef) const;
+	const FTransform GetViewportObjectLocalSpaceTransform(const FAnimDynamicsViewportObjectReference* const SelectedObjectRef) const;
+	const FAnimDynamicsViewportObjectReference* const GetActiveViewportObject() const;
+
+private:
+
+	TArray< FAnimDynamicsViewportObjectReference > SelectedViewportObjects;
+
+	mutable UE::Widget::EWidgetMode CurWidgetMode;
+
+	// Set true every frame where DoTranslation fn is called, reset in Tick fn.
+	bool bIsInteractingWithWidget;
+};
diff --git a/Engine/Source/Editor/AnimGraph/Private/EditModes/AnimNodeEditModes.cpp b/Engine/Source/Editor/AnimGraph/Private/EditModes/AnimNodeEditModes.cpp
--- a/Engine/Source/Editor/AnimGraph/Private/EditModes/AnimNodeEditModes.cpp
+++ b/Engine/Source/Editor/AnimGraph/Private/EditModes/AnimNodeEditModes.cpp
@@ -2,6 +2,7 @@
 
 #include "AnimNodeEditModes.h"
 
+const FEditorModeID AnimNodeEditModes::AnimDynamics("AnimGraph.AnimDynamics");
 const FEditorModeID AnimNodeEditModes::AnimNode("AnimGraph.AnimNode");
 const FEditorModeID AnimNodeEditModes::TwoBoneIK("AnimGraph.SkeletalControl.TwoBoneIK");
 const FEditorModeID AnimNodeEditModes::ObserveBone("AnimGraph.SkeletalControl.ObserveBone");
diff --git a/Engine/Source/Editor/AnimGraph/Public/AnimGraphNode_AnimDynamics.h b/Engine/Source/Editor/AnimGraph/Public/AnimGraphNode_AnimDynamics.h
--- a/Engine/Source/Editor/AnimGraph/Public/AnimGraphNode_AnimDynamics.h
+++ b/Engine/Source/Editor/AnimGraph/Public/AnimGraphNode_AnimDynamics.h
@@ -21,17 +21,21 @@ namespace AnimDynamicsNodeConstants
 {
 	const FLinearColor ShapeDrawColor = FLinearColor::White;
 	const FLinearColor ActiveBodyDrawColor = FLinearColor::Yellow;
-	const float ShapeLineWidth = 0.07f;
-	const float BodyLineWidth = 0.07f;
+	const FLinearColor LimitLineDrawColor = FLinearColor::Yellow;
+	const float LimitPlaneDrawSize = 50.0f;
+	const float ShapeLineWidth = 0.2f;
+	const float SelectedShapeLineWidth = 0.4f;
+	const float BodyLineWidth = 0.2f;
 	const float TransformLineWidth = 0.05f;
 	const float TransformBasisScale = 10.0f;
 }
 
-UCLASS()
+UCLASS(MinimalAPI)
 class UAnimGraphNode_AnimDynamics : public UAnimGraphNode_SkeletalControlBase
 {
 	GENERATED_BODY()
 
+public:
 	UPROPERTY(EditAnywhere, Category = Settings)
 	FAnimNode_AnimDynamics Node;
 
@@ -62,8 +66,6 @@ class UAnimGraphNode_AnimDynamics : public UAnimGraphNode_SkeletalControlBase
 	UPROPERTY(Transient)
 	mutable TObjectPtr<USkeletalMeshComponent> LastPreviewComponent;
 
-public:
-
 	virtual void PostLoad() override;
 
 	static FReply ResetButtonClicked(IDetailLayoutBuilder* DetailLayoutBuilder);
@@ -78,19 +80,17 @@ public:
 
 	// UAnimGraphNode_Base
 	virtual void ValidateAnimNodeDuringCompilation(USkeleton* ForSkeleton, FCompilerResultsLog& MessageLog) override;
+	virtual FEditorModeID GetEditorMode() const override;
 	virtual void CustomizeDetails(IDetailLayoutBuilder& DetailBuilder) override;
-
+	
 	// UAnimGraphNode_SkeletalControlBase
-	virtual void Draw(FPrimitiveDrawInterface* PDI, USkeletalMeshComponent* PreviewSkelMeshComp) const override;
 	virtual void GetOnScreenDebugInfo(TArray<FText>& DebugInfo, FAnimNode_Base* RuntimeAnimNode, USkeletalMeshComponent* PreviewSkelMeshComp) const override;
 
 	FAnimNode_AnimDynamics* GetPreviewDynamicsNode() const;
+	bool IsPreviewLiveActive() const { return bPreviewLive;  }
 
 protected:
 
-	// Keep a version of the current shape for rendering
-	FAnimPhysShape EditPreviewShape;
-
 	virtual FText GetControllerDescription() const override;
 
 	virtual void PostEditChangeProperty(struct FPropertyChangedEvent& PropertyChangedEvent) override;
diff --git a/Engine/Source/Editor/AnimGraph/Public/AnimNodeEditMode.h b/Engine/Source/Editor/AnimGraph/Public/AnimNodeEditMode.h
--- a/Engine/Source/Editor/AnimGraph/Public/AnimNodeEditMode.h
+++ b/Engine/Source/Editor/AnimGraph/Public/AnimNodeEditMode.h
@@ -69,6 +69,8 @@ protected:
 	virtual UAnimGraphNode_Base* GetActiveWidgetAnimNode() const; // Return the editor node associated with the selected widget. All widget operations are performed on this node.
 	virtual FAnimNode_Base*	GetActiveWidgetRuntimeAnimNode() const; // Return the runtime node associated with the selected widget. All widget operations are performed on this node.
 
+	const bool IsManipulatingWidget() const { return bManipulating; }
+
 protected:
 
 	struct EditorRuntimeNodePair
diff --git a/Engine/Source/Editor/AnimGraph/Public/AnimNodeEditModes.h b/Engine/Source/Editor/AnimGraph/Public/AnimNodeEditModes.h
--- a/Engine/Source/Editor/AnimGraph/Public/AnimNodeEditModes.h
+++ b/Engine/Source/Editor/AnimGraph/Public/AnimNodeEditModes.h
@@ -7,6 +7,7 @@
 
 struct ANIMGRAPH_API AnimNodeEditModes
 {
+	const static FEditorModeID AnimDynamics;
 	const static FEditorModeID AnimNode;
 	const static FEditorModeID TwoBoneIK;
 	const static FEditorModeID ObserveBone;
diff --git a/Engine/Source/Editor/AnimationBlueprintEditor/Private/AnimationBlueprintEditor.cpp b/Engine/Source/Editor/AnimationBlueprintEditor/Private/AnimationBlueprintEditor.cpp
--- a/Engine/Source/Editor/AnimationBlueprintEditor/Private/AnimationBlueprintEditor.cpp
+++ b/Engine/Source/Editor/AnimationBlueprintEditor/Private/AnimationBlueprintEditor.cpp
@@ -137,6 +137,22 @@ template <typename TContainerType, typename TPredicate> void SortedContainerDiff
 	}
 }
 
+void FAnimationBlueprintEditor::NotifyAllNodesOnSelection(const bool bInIsSelected)
+{
+	FEditorModeTools& ModeTools = GetEditorModeManager();
+
+	for (TWeakObjectPtr< class UAnimGraphNode_Base > CurrentAnimGraphNode : SelectedAnimGraphNodes)
+	{
+		UAnimGraphNode_Base* const CurrentAnimGraphNodePtr = CurrentAnimGraphNode.Get();
+		FAnimNode_Base* const PreviewNode = FindAnimNode(CurrentAnimGraphNodePtr);
+		
+		// Note: Potentially passing a null PreviewNode ptr when bInIsSelected is false is required to de-select nodes that no longer exist.
+		if (CurrentAnimGraphNodePtr && (!bInIsSelected || PreviewNode))
+		{
+			CurrentAnimGraphNodePtr->OnNodeSelected(bInIsSelected, ModeTools, PreviewNode);
+		}
+	}
+}
 
 /////////////////////////////////////////////////////
 // SAnimBlueprintPreviewPropertyEditor
@@ -1478,6 +1494,8 @@ void FAnimationBlueprintEditor::ClearSelectedActor()
 
 void FAnimationBlueprintEditor::ClearSelectedAnimGraphNodes()
 {
+	NotifyAllNodesOnSelection(false);
+
 	SelectedAnimGraphNodes.Empty();
 }
 
@@ -1770,15 +1788,9 @@ void FAnimationBlueprintEditor::OnSelectedNodesChangedImpl(const TSet<class UObj
 
 			SortedContainerDifference(OldSelectionSorted, NewSelectionSorted, AddSelection, RemSelection, SortPredicate);
 		}
+		
+		NotifyAllNodesOnSelection(false); // Register de-selection with all the previously selected nodes.
 
-		// Register de-selection with all the previously selected nodes.
-		for (FSelectedNodePtr CurrentAnimGraphNode : SelectedAnimGraphNodes)
-		{
-			UAnimGraphNode_Base* const CurrentAnimGraphNodePtr = CurrentAnimGraphNode.Get();
-			FAnimNode_Base* const PreviewNode = FindAnimNode(CurrentAnimGraphNodePtr);
-			// intentionally not null checking PreviewNode here, in order to deselect nodes that are no longer included in the runtime graph after recompile
-			CurrentAnimGraphNodePtr->OnNodeSelected(false, *PersonaEditorModeManager, PreviewNode);
-		}
 
 		// Remove all the nodes that are no longer selected.
 		for (FSelectedNodePtr CurrentAnimGraphNode : RemSelection)
@@ -1792,15 +1804,7 @@ void FAnimationBlueprintEditor::OnSelectedNodesChangedImpl(const TSet<class UObj
 			SelectedAnimGraphNodes.Add(CurrentAnimGraphNode);
 		}
 
-		// Register re-selection with all the currently selected nodes.
-		for (FSelectedNodePtr CurrentAnimGraphNode : SelectedAnimGraphNodes)
-		{
-			UAnimGraphNode_Base* const CurrentAnimGraphNodePtr = CurrentAnimGraphNode.Get();
-			if (FAnimNode_Base* const PreviewNode = FindAnimNode(CurrentAnimGraphNodePtr))
-			{
-				CurrentAnimGraphNodePtr->OnNodeSelected(true, *PersonaEditorModeManager, PreviewNode);
-			}
-		}
+		NotifyAllNodesOnSelection(true); // Register re-selection with all the currently selected nodes.
 	}
 
 	bSelectRegularNode = false;
diff --git a/Engine/Source/Editor/AnimationBlueprintEditor/Private/AnimationBlueprintEditor.h b/Engine/Source/Editor/AnimationBlueprintEditor/Private/AnimationBlueprintEditor.h
--- a/Engine/Source/Editor/AnimationBlueprintEditor/Private/AnimationBlueprintEditor.h
+++ b/Engine/Source/Editor/AnimationBlueprintEditor/Private/AnimationBlueprintEditor.h
@@ -270,6 +270,9 @@ private:
 	/** Called post compile to copy node data */
 	void OnPostCompile();
 
+	/** Call OnNodeSelected for each selected node **/
+	void NotifyAllNodesOnSelection(const bool bInIsSelected);
+
 	/** Helper function used to keep skeletal controls in preview & instance in sync */
 	struct FAnimNode_Base* FindAnimNode(class UAnimGraphNode_Base* AnimGraphNode) const;
 
diff --git a/Engine/Source/Runtime/AnimGraphRuntime/Private/BoneControllers/AnimNode_AnimDynamics.cpp b/Engine/Source/Runtime/AnimGraphRuntime/Private/BoneControllers/AnimNode_AnimDynamics.cpp
--- a/Engine/Source/Runtime/AnimGraphRuntime/Private/BoneControllers/AnimNode_AnimDynamics.cpp
+++ b/Engine/Source/Runtime/AnimGraphRuntime/Private/BoneControllers/AnimNode_AnimDynamics.cpp
@@ -20,6 +20,35 @@ TAutoConsoleVariable<int32> CVarEnableAdaptiveSubstep(TEXT("p.AnimDynamicsAdapti
 TAutoConsoleVariable<int32> CVarAdaptiveSubstepNumDebtFrames(TEXT("p.AnimDynamicsNumDebtFrames"), 5, TEXT("Number of frames to maintain as time debt when using adaptive substepping, this should be at least 1 or the time debt will never be cleared."));
 TAutoConsoleVariable<int32> CVarEnableWind(TEXT("p.AnimDynamicsWind"), 1, TEXT("Enables/Disables anim dynamics wind forces globally."), ECVF_Scalability);
 
+// FindChainBones
+// 
+// Returns an ordered list of the names of all the bones in a chain between BeginBoneName and EndBoneName.
+//
+void FindChainBones(const FName BeginBoneName, const FName EndBoneName, TFunctionRef< FName (const FName) > GetParentBoneName, TArray<FName>& OutChainBoneNames)
+{
+	OutChainBoneNames.Add(BeginBoneName);
+
+	const uint32 OutputStartIndex = OutChainBoneNames.Num();
+
+	// Add the end of the chain. We have to walk from the bottom upwards to find a chain
+	// as walking downwards doesn't guarantee a single end point.
+
+	// Walk up the chain until we either find the top or hit the root bone
+	FName ParentBoneName = EndBoneName;
+	for (; !ParentBoneName.IsNone() && (ParentBoneName != BeginBoneName); ParentBoneName = GetParentBoneName(ParentBoneName))
+	{
+		OutChainBoneNames.Insert(ParentBoneName, OutputStartIndex);
+	}
+
+	if (ParentBoneName != BeginBoneName)
+	{
+		UE_LOG(LogAnimation, Error, TEXT("AnimDynamics: Attempted to find bone chain starting at %s and ending at %s but failed."), *BeginBoneName.ToString(), *EndBoneName.ToString());
+
+		// Remove any accumulated chain bone names beyond the start bone from output.
+		OutChainBoneNames.RemoveAt(OutputStartIndex, OutChainBoneNames.Num() - OutputStartIndex);
+	}
+}
+
 const float FAnimNode_AnimDynamics::MaxTimeDebt = (1.0f / 60.0f) * 5.0f; // 5 frames max debt
 
 #if ENABLE_ANIM_DRAW_DEBUG
@@ -56,16 +85,20 @@ void FAnimNode_AnimDynamics::DrawBodies(FComponentSpacePoseContext& InContext, c
 	const bool bFilterBone = FilteredBoneName.Len() > 0;
 
 	const int32 NumBodies = Bodies.Num();
+
+	check(PhysicsBodyDefinitions.Num() == NumBodies);
+
 	for(int32 BodyIndex = 0 ; BodyIndex < NumBodies ; ++BodyIndex)
 	{
 		const FAnimPhysRigidBody& Body = Bodies[BodyIndex].RigidBody.PhysBody;
+		const FAnimPhysBodyDefinition& PhysicsBodyDef = PhysicsBodyDefinitions[BodyIndex];
 
-		if(bFilterBone && BoundBoneReferences[BodyIndex].BoneName != FName(*FilteredBoneName))
+		if(bFilterBone && PhysicsBodyDef.BoundBone.BoneName != FName(*FilteredBoneName))
 		{
 			continue;
 		}
 
-		FTransform Transform(Body.Pose.Orientation, Body.Pose.Position + Body.Pose.Orientation.RotateVector(JointOffsets[BodyIndex]));
+		FTransform Transform(Body.Pose.Orientation, Body.Pose.Position);
 		Transform = GetComponentSpaceTransformFromSimSpace(SimulationSpace, InContext, Transform);
 		Transform *= Proxy->GetComponentTransform();
 
@@ -126,8 +159,6 @@ FAnimNode_AnimDynamics::FAnimNode_AnimDynamics()
 , AngularDampingOverride(0.0f)
 , PreviousCompWorldSpaceTM(FTransform::Identity)
 , PreviousActorWorldSpaceTM(FTransform::Identity)
-, BoxExtents(0.0f)
-, LocalJointOffset(0.0f)
 , GravityScale(1.0f)
 , GravityOverride(ForceInitToZero)
 , LinearSpringConstant(0.0f)
@@ -136,9 +167,7 @@ FAnimNode_AnimDynamics::FAnimNode_AnimDynamics()
 , AngularBiasOverride(0.0f)
 , NumSolverIterationsPreUpdate(4)
 , NumSolverIterationsPostUpdate(1)
-, SphereCollisionRadius(0.0f)
 , ExternalForce(ForceInitToZero)
-, CollisionType(AnimPhysCollisionType::CoM)
 , SimulationSpace(AnimPhysSimSpaceType::Component)
 , LastSimSpace(AnimPhysSimSpaceType::Component)
 , InitTeleportType(ETeleportType::None)
@@ -157,6 +186,15 @@ FAnimNode_AnimDynamics::FAnimNode_AnimDynamics()
 , bAngularSpring(false)
 , bChain(false)
 , RetargetingSettings(FRotationRetargetingInfo(false /* enabled */))
+#if WITH_EDITORONLY_DATA
+, BoxExtents_DEPRECATED(FVector::ZeroVector)
+, LocalJointOffset_DEPRECATED(FVector::ZeroVector)
+, CollisionType_DEPRECATED(AnimPhysCollisionType::CoM)
+, SphereCollisionRadius_DEPRECATED(0.0f)
+#endif
+#if WITH_EDITOR
+, bDoPhysicsUpdateInEditor(true)
+#endif
 #if ENABLE_ANIM_DRAW_DEBUG
 , FilteredBoneIndex(INDEX_NONE)
 #endif
@@ -166,6 +204,8 @@ FAnimNode_AnimDynamics::FAnimNode_AnimDynamics()
 	ComponentAppliedLinearAccClamp = FVector(100000, 100000, 100000);
 
 	PreviousComponentLinearVelocity = FVector::ZeroVector;
+
+	PhysicsBodyDefinitions.Add(FAnimPhysBodyDefinition());
 }
 
 void FAnimNode_AnimDynamics::Initialize_AnyThread(const FAnimationInitializeContext& Context)
@@ -250,7 +290,7 @@ void FAnimNode_AnimDynamics::EvaluateSkeletalControl_AnyThread(FComponentSpacePo
 			}
 		}
 
-		if (bDoUpdate && NextTimeStep > AnimPhysicsMinDeltaTime)
+		if (ShouldDoPhysicsUpdate() && NextTimeStep > AnimPhysicsMinDeltaTime)
 		{
 			// Calculate gravity direction
 			SimSpaceGravityDirection = TransformWorldVectorToSimSpace(Output, FVector(0.0f, 0.0f, -1.0f));
@@ -355,9 +395,12 @@ void FAnimNode_AnimDynamics::EvaluateSkeletalControl_AnyThread(FComponentSpacePo
 
 			const FBoneContainer& BoneContainer = Output.Pose.GetPose().GetBoneContainer();
 
-			for (int32 Idx = 0; Idx < BoundBoneReferences.Num(); ++Idx)
+			check(Bodies.Num() == PhysicsBodyDefinitions.Num());
+			check(PhysicsBodyJointOffsets.Num() == PhysicsBodyDefinitions.Num());
+
+			for (int32 Idx = 0; Idx < Bodies.Num(); ++Idx)
 			{
-				FBoneReference& CurrentChainBone = BoundBoneReferences[Idx];
+				FBoneReference& CurrentChainBone = PhysicsBodyDefinitions[Idx].BoundBone;
 				FAnimPhysRigidBody& CurrentBody = Bodies[Idx].RigidBody.PhysBody;
 
 				// Skip invalid bones
@@ -368,7 +411,8 @@ void FAnimNode_AnimDynamics::EvaluateSkeletalControl_AnyThread(FComponentSpacePo
 
 				FCompactPoseBoneIndex BoneIndex = CurrentChainBone.GetCompactPoseIndex(BoneContainer);
 
-				FTransform NewBoneTransform(CurrentBody.Pose.Orientation, CurrentBody.Pose.Position + CurrentBody.Pose.Orientation.RotateVector(JointOffsets[Idx]));
+				// Calculate target bone transform from physics body.
+				FTransform NewBoneTransform(CurrentBody.Pose.Orientation, CurrentBody.Pose.Position - CurrentBody.Pose.Orientation.RotateVector(PhysicsBodyJointOffsets[Idx]));
 
 				if (RetargetingSettings.bEnabled)
 				{
@@ -404,7 +448,7 @@ void FAnimNode_AnimDynamics::EvaluateSkeletalControl_AnyThread(FComponentSpacePo
 			}
 		}
 
-		// Store our sim space incase it changes
+		// Store our sim space in case it changes
 		LastSimSpace = SimulationSpace;
 
 		// Store previous component and actor space transform
@@ -413,6 +457,68 @@ void FAnimNode_AnimDynamics::EvaluateSkeletalControl_AnyThread(FComponentSpacePo
 	}
 }
 
+void FAnimNode_AnimDynamics::UpdateChainPhysicsBodyDefinitions(const USkeletalMeshComponent* const PreviewSkelMeshComp)
+{
+	if (PreviewSkelMeshComp)
+	{
+		TArray<FName> ChainBoneNames;
+		FindChainBones(BoundBone.BoneName, ChainEnd.BoneName, [PreviewSkelMeshComp](const FName BoneName) { return PreviewSkelMeshComp->GetParentBone(BoneName); }, ChainBoneNames);
+		UpdateChainPhysicsBodyDefinitions(ChainBoneNames);
+	}
+}
+
+void FAnimNode_AnimDynamics::UpdateChainPhysicsBodyDefinitions(const FReferenceSkeleton& ReferenceSkeleton)
+{
+	auto GetParentBoneNameFn = [&ReferenceSkeleton](const FName BoneName)
+	{
+		int32 BoneIndex = ReferenceSkeleton.FindBoneIndex(BoneName);
+		BoneIndex = ReferenceSkeleton.GetParentIndex(BoneIndex);
+		return ReferenceSkeleton.GetBoneName(BoneIndex);
+	};
+
+	TArray<FName> ChainBoneNames;
+	FindChainBones(BoundBone.BoneName, ChainEnd.BoneName, GetParentBoneNameFn, ChainBoneNames);
+	UpdateChainPhysicsBodyDefinitions(ChainBoneNames);
+}
+
+void FAnimNode_AnimDynamics::UpdateChainPhysicsBodyDefinitions(const TArray<FName>& ChainBoneNames)
+{
+	check(ChainBoneNames.Num() > 0);
+
+	// Remove any bodies for bones that are not in the chain.
+	PhysicsBodyDefinitions.RemoveAll([&ChainBoneNames](const FAnimPhysBodyDefinition& Value) { return ChainBoneNames.Find(Value.BoundBone.BoneName) == INDEX_NONE;});
+	PhysicsBodyDefinitions.Reserve(ChainBoneNames.Num());
+
+	// Create a new Physics Body Def for any new bones in the chain and add them before or after the existing bodies as appropriate to maintain the order of the chain bones.
+	{
+		// If there was only one physics body then copy its values to all new chain bodies (emulating legacy behaviour), otherwise use values from default construction.
+		FAnimPhysBodyDefinition PrototypePhysBodyDef;
+		
+		if (PhysicsBodyDefinitions.Num() == 1)
+		{
+			PrototypePhysBodyDef = PhysicsBodyDefinitions[0];
+		}
+
+		uint32 PhysicsBodyDefIndex = 0;
+
+		for (FName BoneName : ChainBoneNames)
+		{
+			if (!PhysicsBodyDefinitions.FindByPredicate([BoneName](const FAnimPhysBodyDefinition& Value) { return Value.BoundBone.BoneName == BoneName; }))
+			{
+				// Add the new bone to the chain.
+				FAnimPhysBodyDefinition ChainPhysicsBody;
+				ChainPhysicsBody.BoundBone.BoneName = BoneName;
+				ChainPhysicsBody.BoxExtents = PrototypePhysBodyDef.BoxExtents;
+				ChainPhysicsBody.LocalJointOffset = PrototypePhysBodyDef.LocalJointOffset;
+				ChainPhysicsBody.SphereCollisionRadius = PrototypePhysBodyDef.SphereCollisionRadius;
+				PhysicsBodyDefinitions.Insert(ChainPhysicsBody, PhysicsBodyDefIndex);
+			}
+
+			++PhysicsBodyDefIndex;
+		}
+	}
+}
+
 void FAnimNode_AnimDynamics::InitializeBoneReferences(const FBoneContainer& RequiredBones)
 {
 	DECLARE_SCOPE_HIERARCHICAL_COUNTER_ANIMNODE(InitializeBoneReferences)
@@ -441,10 +547,10 @@ void FAnimNode_AnimDynamics::InitializeBoneReferences(const FBoneContainer& Requ
 	// If we're currently simulating (LOD change etc.)
 	bool bSimulating = ActiveBoneIndices.Num() > 0;
 
-	const int32 NumRefs = BoundBoneReferences.Num();
+	const int32 NumRefs = PhysicsBodyDefinitions.Num();
 	for(int32 BoneRefIdx = 0; BoneRefIdx < NumRefs; ++BoneRefIdx)
 	{
-		FBoneReference& BoneRef = BoundBoneReferences[BoneRefIdx];
+		FBoneReference& BoneRef = PhysicsBodyDefinitions[BoneRefIdx].BoundBone;
 		BoneRef.Initialize(RequiredBones);
 
 		if(bSimulating)
@@ -498,9 +604,9 @@ bool FAnimNode_AnimDynamics::IsValidToEvaluate(const USkeleton* Skeleton, const
 		{
 			// Check for LOD subchain
 			int32 NumValidBonesFromRoot = 0;
-			for(FBoneReference& BoneRef : BoundBoneReferences)
+			for(const FAnimPhysBodyDefinition& PhysicsBodyDef : PhysicsBodyDefinitions)
 			{
-				if(BoneRef.IsValidToEvaluate(RequiredBones))
+				if(PhysicsBodyDef.BoundBone.IsValidToEvaluate(RequiredBones))
 				{
 					bSubChainValid = true;
 					break;
@@ -524,28 +630,19 @@ const FAnimPhysRigidBody& FAnimNode_AnimDynamics::GetPhysBody(int32 BodyIndex) c
 	return Bodies[BodyIndex].RigidBody.PhysBody;
 }
 
-#if WITH_EDITOR
-FVector FAnimNode_AnimDynamics::GetBodyLocalJointOffset(int32 BodyIndex) const
-{
-	if (JointOffsets.IsValidIndex(BodyIndex))
-	{
-		return JointOffsets[BodyIndex];
-	}
-	return FVector::ZeroVector;
-}
-
-int32 FAnimNode_AnimDynamics::GetNumBoundBones() const
+FTransform FAnimNode_AnimDynamics::GetBodyComponentSpaceTransform(const FAnimPhysRigidBody& Body, const USkeletalMeshComponent* const SkelComp) const
 {
-	return BoundBoneReferences.Num();
+	return GetComponentSpaceTransformFromSimSpace(SimulationSpace, SkelComp, FTransform(Body.Pose.Orientation, Body.Pose.Position));
 }
 
-const FBoneReference* FAnimNode_AnimDynamics::GetBoundBoneReference(int32 Index) const
+#if WITH_EDITOR
+FVector FAnimNode_AnimDynamics::GetBodyLocalJointOffset(const int32 BodyIndex) const
 {
-	if(BoundBoneReferences.IsValidIndex(Index))
+	if (PhysicsBodyDefinitions.IsValidIndex(BodyIndex))
 	{
-		return &BoundBoneReferences[Index];
+		return PhysicsBodyDefinitions[BodyIndex].LocalJointOffset;
 	}
-	return nullptr;
+	return FVector::ZeroVector;
 }
 
 #endif
@@ -558,215 +655,165 @@ void FAnimNode_AnimDynamics::RequestInitialise(ETeleportType InTeleportType)
 
 void FAnimNode_AnimDynamics::InitPhysics(FComponentSpacePoseContext& Output)
 {
-	switch(InitTeleportType)
+	switch (InitTeleportType)
 	{
-		case ETeleportType::ResetPhysics:
-		{
-			// Clear up any existing physics data
-			TermPhysics();
+	case ETeleportType::ResetPhysics:
+	{
+		// Clear up any existing physics data
+		TermPhysics();
 
-			const FBoneContainer& BoneContainer = Output.Pose.GetPose().GetBoneContainer();
+		const FBoneContainer& BoneContainer = Output.Pose.GetPose().GetBoneContainer();
 
-	
-			// List of bone indices in the chain.
-			TArray<int32> ChainBoneIndices;
-			TArray<FName> ChainBoneNames;
+		if (PhysicsBodyDefinitions.Num() && ((PhysicsBodyDefinitions[0].BoundBone.BoneName != BoundBone.BoneName) || (PhysicsBodyDefinitions.Last().BoundBone.BoneName != ChainEnd.BoneName)))
+		{
+			UpdateChainPhysicsBodyDefinitions(BoneContainer.GetReferenceSkeleton());
+		}
 
-			if(ChainEnd.IsValidToEvaluate(BoneContainer))
-			{
-				// Add the end of the chain. We have to walk from the bottom upwards to find a chain
-				// as walking downwards doesn't guarantee a single end point.
+		// Transform GravityOverride to simulation space if necessary.
+		const FVector GravityOverrideSimSpace = (bUseGravityOverride && !bGravityOverrideInSimSpace) ? TransformWorldVectorToSimSpace(Output, GravityOverride) : GravityOverride;
+		
+		check(PhysicsBodyDefinitions.Num() > 0);
+		if (PhysicsBodyDefinitions.Num() > 0)
+		{
+			Bodies.Reserve(PhysicsBodyDefinitions.Num());
+			PhysicsBodyDefinitions[0].BoundBone = BoundBone;
+		}
 
-				ChainBoneIndices.Add(ChainEnd.BoneIndex);
-				ChainBoneNames.Add(ChainEnd.BoneName);
+		ConstraintOffsets.Reset(PhysicsBodyDefinitions.Num());
+		PhysicsBodyJointOffsets.Reset(PhysicsBodyDefinitions.Num());
 
-				int32 ParentBoneIndex = BoneContainer.GetParentBoneIndex(ChainEnd.BoneIndex);
+		FTransform PreviousBodyTransformSimSpace;
+		PreviousBodyTransformSimSpace.SetIdentity();
 
-				// Walk up the chain until we either find the top or hit the root bone
-				while(ParentBoneIndex > 0)
-				{
-					ChainBoneIndices.Add(ParentBoneIndex);
-					ChainBoneNames.Add(BoneContainer.GetReferenceSkeleton().GetBoneName(ParentBoneIndex));
+		for (FAnimPhysBodyDefinition& PhysicsBodyDef : PhysicsBodyDefinitions)
+		{
+			TArray<FAnimPhysShape> BodyShapes;
+			BodyShapes.Add(FAnimPhysShape::MakeBox(PhysicsBodyDef.BoxExtents));
 
-					if(ParentBoneIndex == BoundBone.BoneIndex)
-					{
-						// Found the top of the chain
-						break;
-					}
+			PhysicsBodyDef.BoundBone.Initialize(BoneContainer);
 
-					ParentBoneIndex = BoneContainer.GetParentBoneIndex(ParentBoneIndex);
-				}
+			FTransform BodyTransform = GetBoneTransformInSimSpace(Output, PhysicsBodyDef.BoundBone.GetCompactPoseIndex(BoneContainer));
+			BodyTransform.SetTranslation(BodyTransform.GetTranslation() + BodyTransform.GetRotation().RotateVector(PhysicsBodyDef.LocalJointOffset)); // Transform for physics body in Sim Space.
 
-				// Bail if we can't find a chain, and let the user know
-				if(ParentBoneIndex != BoundBone.BoneIndex)
-				{
-					UE_LOG(LogAnimation, Error, TEXT("AnimDynamics: Attempted to find bone chain starting at %s and ending at %s but failed."), *BoundBone.BoneName.ToString(), *ChainEnd.BoneName.ToString());
-					return;
-				}
-			}
-			else
+			FAnimPhysLinkedBody NewChainBody(BodyShapes, BodyTransform.GetTranslation(), PhysicsBodyDef.BoundBone);
+			FAnimPhysRigidBody& PhysicsBody = NewChainBody.RigidBody.PhysBody;
+			PhysicsBody.Pose.Orientation = BodyTransform.GetRotation();
+			PhysicsBody.PreviousOrientation = PhysicsBody.Pose.Orientation;
+			PhysicsBody.NextOrientation = PhysicsBody.Pose.Orientation;
+			PhysicsBody.CollisionType = PhysicsBodyDef.CollisionType;
+
+			switch (PhysicsBody.CollisionType)
 			{
-				// No chain specified, just use the bound bone
-				ChainBoneIndices.Add(BoundBone.BoneIndex);
-				ChainBoneNames.Add(BoundBone.BoneName);
+			case AnimPhysCollisionType::CustomSphere:
+				PhysicsBody.SphereCollisionRadius = PhysicsBodyDef.SphereCollisionRadius;
+				break;
+			case AnimPhysCollisionType::InnerSphere:
+				PhysicsBody.SphereCollisionRadius = PhysicsBodyDef.BoxExtents.GetAbsMin() / 2.0f;
+				break;
+			case AnimPhysCollisionType::OuterSphere:
+				PhysicsBody.SphereCollisionRadius = PhysicsBodyDef.BoxExtents.GetAbsMax() / 2.0f;
+				break;
+			default:
+				break;
 			}
 
-			// Transform GravityOverride to simulation space if necessary.
-			const FVector GravityOverrideSimSpace = (bUseGravityOverride && !bGravityOverrideInSimSpace) ? TransformWorldVectorToSimSpace(Output, GravityOverride) : GravityOverride;
-
-			Bodies.Reserve(ChainBoneIndices.Num());
-			// Walk backwards here as the chain was discovered in reverse order
-			for (int32 Idx = ChainBoneIndices.Num() - 1; Idx >= 0; --Idx)
+			if (bOverrideLinearDamping)
 			{
-				TArray<FAnimPhysShape> BodyShapes;
-				BodyShapes.Add(FAnimPhysShape::MakeBox(BoxExtents));
-
-				FBoneReference LinkBoneRef;
-				LinkBoneRef.BoneName = ChainBoneNames[Idx];
-				LinkBoneRef.Initialize(BoneContainer);
-
-				// Calculate joint offsets by looking at the length of the bones and extending the provided offset
-				if (BoundBoneReferences.Num() > 0)
-				{
-					FTransform CurrentBoneTransform = GetBoneTransformInSimSpace(Output, LinkBoneRef.GetCompactPoseIndex(BoneContainer));
-					FTransform PreviousBoneTransform = GetBoneTransformInSimSpace(Output, BoundBoneReferences.Last().GetCompactPoseIndex(BoneContainer));
-
-					FVector PreviousAnchor = PreviousBoneTransform.TransformPosition(-LocalJointOffset);
-					float DistanceToAnchor = (PreviousBoneTransform.GetTranslation() - CurrentBoneTransform.GetTranslation()).Size() * 0.5f;
-
-					if(LocalJointOffset.SizeSquared() < SMALL_NUMBER)
-					{
-						// No offset, just use the position between chain links as the offset
-						// This is likely to just look horrible, but at least the bodies will
-						// be placed correctly and not stack up at the top of the chain.
-						JointOffsets.Add(PreviousAnchor - CurrentBoneTransform.GetTranslation());
-					}
-					else
-					{
-						// Extend offset along chain.
-						JointOffsets.Add(LocalJointOffset.GetSafeNormal() * DistanceToAnchor);
-					}
-				}
-				else
-				{
-					// No chain to worry about, just use the specified offset.
-					JointOffsets.Add(LocalJointOffset);
-				}
-
-				BoundBoneReferences.Add(LinkBoneRef);
-
-				FTransform BodyTransform = GetBoneTransformInSimSpace(Output, LinkBoneRef.GetCompactPoseIndex(BoneContainer));
-
-				BodyTransform.SetTranslation(BodyTransform.GetTranslation() + BodyTransform.GetRotation().RotateVector(-LocalJointOffset));
-
-				FAnimPhysLinkedBody NewChainBody(BodyShapes, BodyTransform.GetTranslation(), LinkBoneRef);
-				FAnimPhysRigidBody& PhysicsBody = NewChainBody.RigidBody.PhysBody;
-				PhysicsBody.Pose.Orientation = BodyTransform.GetRotation();
-				PhysicsBody.PreviousOrientation = PhysicsBody.Pose.Orientation;
-				PhysicsBody.NextOrientation = PhysicsBody.Pose.Orientation;
-				PhysicsBody.CollisionType = CollisionType;
-
-				switch(PhysicsBody.CollisionType)
-				{
-					case AnimPhysCollisionType::CustomSphere:
-						PhysicsBody.SphereCollisionRadius = SphereCollisionRadius;
-						break;
-					case AnimPhysCollisionType::InnerSphere:
-						PhysicsBody.SphereCollisionRadius = BoxExtents.GetAbsMin() / 2.0f;
-						break;
-					case AnimPhysCollisionType::OuterSphere:
-						PhysicsBody.SphereCollisionRadius = BoxExtents.GetAbsMax() / 2.0f;
-						break;
-					default:
-						break;
-				}
+				PhysicsBody.bLinearDampingOverriden = true;
+				PhysicsBody.LinearDamping = LinearDampingOverride;
+			}
 
-				if (bOverrideLinearDamping)
-				{
-					PhysicsBody.bLinearDampingOverriden = true;
-					PhysicsBody.LinearDamping = LinearDampingOverride;
-				}
+			if (bOverrideAngularDamping)
+			{
+				PhysicsBody.bAngularDampingOverriden = true;
+				PhysicsBody.AngularDamping = AngularDampingOverride;
+			}
 
-				if (bOverrideAngularDamping)
-				{
-					PhysicsBody.bAngularDampingOverriden = true;
-					PhysicsBody.AngularDamping = AngularDampingOverride;
-				}
+			PhysicsBody.GravityScale = GravityScale;
+			PhysicsBody.bUseGravityOverride = bUseGravityOverride;
+			PhysicsBody.GravityOverride = GravityOverrideSimSpace;
 
-				PhysicsBody.GravityScale = GravityScale;
-				PhysicsBody.bUseGravityOverride = bUseGravityOverride;
-				PhysicsBody.GravityOverride = GravityOverrideSimSpace;
-				
-				PhysicsBody.bWindEnabled = bWindWasEnabled;
+			PhysicsBody.bWindEnabled = bWindWasEnabled;
 
+			if (Bodies.Num() > 0)
+			{
 				// Link to parent
-				if (Bodies.Num() > 0)
-				{
-					NewChainBody.ParentBody = &Bodies.Last().RigidBody;
-				}
+				NewChainBody.ParentBody = &Bodies.Last().RigidBody;
 
-				Bodies.Add(NewChainBody);
-				ActiveBoneIndices.Add(Bodies.Num() - 1);
+				// Calculate constraint offset positions in the space of each body.
+				const FVector ConstaintLocationSimSpace = (BodyTransform.GetTranslation() - PreviousBodyTransformSimSpace.GetTranslation()) * 0.5f;
+				ConstraintOffsets.Add(FAnimConstraintOffsetPair(PreviousBodyTransformSimSpace.GetRotation().UnrotateVector(ConstaintLocationSimSpace), BodyTransform.GetRotation().UnrotateVector(-ConstaintLocationSimSpace)));
 			}
-	
-			BaseBodyPtrs.Empty();
-			for(FAnimPhysLinkedBody& Body : Bodies)
+			else
 			{
-				BaseBodyPtrs.Add(&Body.RigidBody.PhysBody);
+				// The first physics body is constrained to the location of it's bound bone.
+				ConstraintOffsets.Add(FAnimConstraintOffsetPair(FVector::ZeroVector, -PhysicsBodyDef.LocalJointOffset)); // Offset is the vector from the physics body to its associated bone in the bodies local space.
 			}
 
 			// Set up transient constraint data
-			const bool bXAxisLocked = ConstraintSetup.LinearXLimitType != AnimPhysLinearConstraintType::Free && ConstraintSetup.LinearAxesMin.X - ConstraintSetup.LinearAxesMax.X == 0.0f;
-			const bool bYAxisLocked = ConstraintSetup.LinearYLimitType != AnimPhysLinearConstraintType::Free && ConstraintSetup.LinearAxesMin.Y - ConstraintSetup.LinearAxesMax.Y == 0.0f;
-			const bool bZAxisLocked = ConstraintSetup.LinearZLimitType != AnimPhysLinearConstraintType::Free && ConstraintSetup.LinearAxesMin.Z - ConstraintSetup.LinearAxesMax.Z == 0.0f;
+			const bool bXAxisLocked = PhysicsBodyDef.ConstraintSetup.LinearXLimitType != AnimPhysLinearConstraintType::Free && PhysicsBodyDef.ConstraintSetup.LinearAxesMin.X - PhysicsBodyDef.ConstraintSetup.LinearAxesMax.X == 0.0f;
+			const bool bYAxisLocked = PhysicsBodyDef.ConstraintSetup.LinearYLimitType != AnimPhysLinearConstraintType::Free && PhysicsBodyDef.ConstraintSetup.LinearAxesMin.Y - PhysicsBodyDef.ConstraintSetup.LinearAxesMax.Y == 0.0f;
+			const bool bZAxisLocked = PhysicsBodyDef.ConstraintSetup.LinearZLimitType != AnimPhysLinearConstraintType::Free && PhysicsBodyDef.ConstraintSetup.LinearAxesMin.Z - PhysicsBodyDef.ConstraintSetup.LinearAxesMax.Z == 0.0f;
+			PhysicsBodyDef.ConstraintSetup.bLinearFullyLocked = bXAxisLocked && bYAxisLocked && bZAxisLocked;
 
-			ConstraintSetup.bLinearFullyLocked = bXAxisLocked && bYAxisLocked && bZAxisLocked;
+			Bodies.Add(NewChainBody);
+			ActiveBoneIndices.Add(Bodies.Num() - 1);
+			PhysicsBodyJointOffsets.Add(PhysicsBodyDef.LocalJointOffset);
 
-			// Cache physics settings to avoid accessing UPhysicsSettings continuously
-			if(UPhysicsSettings* Settings = UPhysicsSettings::Get())
-			{
-				AnimPhysicsMinDeltaTime = Settings->AnimPhysicsMinDeltaTime;
-				MaxPhysicsDeltaTime = Settings->MaxPhysicsDeltaTime;
-				MaxSubstepDeltaTime = Settings->MaxSubstepDeltaTime;
-				MaxSubsteps = Settings->MaxSubsteps;
-			}
-			else
-			{
-				AnimPhysicsMinDeltaTime = 0.f;
-				MaxPhysicsDeltaTime = (1.0f/30.0f);
-				MaxSubstepDeltaTime = (1.0f/60.0f);
-				MaxSubsteps = 4;
-			}
+			PreviousBodyTransformSimSpace = BodyTransform;
+		}
 
-			SimSpaceGravityDirection = TransformWorldVectorToSimSpace(Output, FVector(0.0f, 0.0f, -1.0f));
+		BaseBodyPtrs.Reset();
+		for (FAnimPhysLinkedBody& Body : Bodies)
+		{
+			BaseBodyPtrs.Add(&Body.RigidBody.PhysBody);
 		}
-		break;
 
-		case ETeleportType::TeleportPhysics:
+		// Cache physics settings to avoid accessing UPhysicsSettings continuously
+		if (UPhysicsSettings* Settings = UPhysicsSettings::Get())
+		{
+			AnimPhysicsMinDeltaTime = Settings->AnimPhysicsMinDeltaTime;
+			MaxPhysicsDeltaTime = Settings->MaxPhysicsDeltaTime;
+			MaxSubstepDeltaTime = Settings->MaxSubstepDeltaTime;
+			MaxSubsteps = Settings->MaxSubsteps;
+		}
+		else
+		{
+			AnimPhysicsMinDeltaTime = 0.f;
+			MaxPhysicsDeltaTime = (1.0f / 30.0f);
+			MaxSubstepDeltaTime = (1.0f / 60.0f);
+			MaxSubsteps = 4;
+		}
+
+		SimSpaceGravityDirection = TransformWorldVectorToSimSpace(Output, FVector(0.0f, 0.0f, -1.0f));
+	}
+	break;
+
+	case ETeleportType::TeleportPhysics:
+	{
+		// Clear any external forces
+		ExternalForce = FVector::ZeroVector;
+
+		// Move any active bones
+		for (const int32& BodyIndex : ActiveBoneIndices)
 		{
-			// Clear any external forces
-			ExternalForce = FVector::ZeroVector;
+			FAnimPhysRigidBody& Body = Bodies[BodyIndex].RigidBody.PhysBody;
 
-			// Move any active bones
-			for(const int32& BodyIndex : ActiveBoneIndices)
-			{
-				FAnimPhysRigidBody& Body = Bodies[BodyIndex].RigidBody.PhysBody;
+			// Get old comp space transform
+			FTransform BodyTransform(Body.Pose.Orientation, Body.Pose.Position - Body.Pose.Orientation.RotateVector(PhysicsBodyDefinitions[BodyIndex].LocalJointOffset));
+			BodyTransform = GetComponentSpaceTransformFromSimSpace(SimulationSpace, Output, BodyTransform, PreviousCompWorldSpaceTM, PreviousActorWorldSpaceTM);
 
-				// Get old comp space transform
-				FTransform BodyTransform(Body.Pose.Orientation, Body.Pose.Position + Body.Pose.Orientation.RotateVector(JointOffsets[BodyIndex]));
-				BodyTransform = GetComponentSpaceTransformFromSimSpace(SimulationSpace, Output, BodyTransform, PreviousCompWorldSpaceTM, PreviousActorWorldSpaceTM);
+			// move to new space
+			BodyTransform = GetSimSpaceTransformFromComponentSpace(SimulationSpace, Output, BodyTransform);
 
-				// move to new space
-				BodyTransform = GetSimSpaceTransformFromComponentSpace(SimulationSpace, Output, BodyTransform);
-				
-				Body.Pose.Orientation = BodyTransform.GetRotation();
-				Body.PreviousOrientation = Body.Pose.Orientation;
-				Body.NextOrientation = Body.Pose.Orientation;
+			Body.Pose.Orientation = BodyTransform.GetRotation();
+			Body.PreviousOrientation = Body.Pose.Orientation;
+			Body.NextOrientation = Body.Pose.Orientation;
 
-				Body.Pose.Position = BodyTransform.GetTranslation() - Body.Pose.Orientation.RotateVector(JointOffsets[BodyIndex]);
-			}
+			Body.Pose.Position = BodyTransform.GetTranslation() + Body.Pose.Orientation.RotateVector(PhysicsBodyDefinitions[BodyIndex].LocalJointOffset);
 		}
-		break;
+	}
+	break;
 	}
 
 	InitTeleportType = ETeleportType::None;
@@ -781,10 +828,14 @@ void FAnimNode_AnimDynamics::TermPhysics()
 	AngularLimits.Reset();
 	Springs.Reset();
 	ActiveBoneIndices.Reset();
-
-	BoundBoneReferences.Reset();
-	JointOffsets.Reset();
+	PhysicsBodyJointOffsets.Reset();
+	ConstraintOffsets.Reset();
 	BodiesToReset.Reset();
+
+	for (FAnimPhysBodyDefinition& PhysicsBodyDef : PhysicsBodyDefinitions)
+	{
+		PhysicsBodyDef.BoundBone.BoneIndex = INDEX_NONE;
+	}
 }
 
 void FAnimNode_AnimDynamics::UpdateLimits(FComponentSpacePoseContext& Output)
@@ -800,7 +851,8 @@ void FAnimNode_AnimDynamics::UpdateLimits(FComponentSpacePoseContext& Output)
 
 	for(int32 ActiveIndex : ActiveBoneIndices)
 	{
-		const FBoneReference& CurrentBoneRef = BoundBoneReferences[ActiveIndex];
+		const FAnimPhysBodyDefinition& PhysicsBodyDef = PhysicsBodyDefinitions[ActiveIndex];
+		const FBoneReference& CurrentBoneRef = PhysicsBodyDef.BoundBone;
 
 		// If our bone isn't valid, move on
 		if(!CurrentBoneRef.IsValidToEvaluate(BoneContainer))
@@ -823,58 +875,55 @@ void FAnimNode_AnimDynamics::UpdateLimits(FComponentSpacePoseContext& Output)
 
 		FTransform ShapeTransform = BoundBoneTransform;
 		
-		// Local offset to joint for Body1
-		FVector Body1JointOffset = LocalJointOffset;
-
 		if (PrevBody)
 		{
-			// Get the correct offset
-			Body1JointOffset = JointOffsets[ActiveIndex];
-			// Modify the shape transform to be correct in Body0 frame
-			ShapeTransform = FTransform(FQuat::Identity, -Body1JointOffset);
+			ShapeTransform = FTransform(FQuat::Identity, ConstraintOffsets[ActiveIndex].Body0Offset);
 		}
-		
-		if (ConstraintSetup.bLinearFullyLocked)
+
+		// Local offset to joint for Body1
+		const FVector Body1JointOffset = ConstraintOffsets[ActiveIndex].Body1Offset;
+
+		if (PhysicsBodyDef.ConstraintSetup.bLinearFullyLocked)
 		{
 			// Rather than calculate prismatic limits, just lock the transform (1 limit instead of 6)
 			FAnimPhys::ConstrainPositionNailed(NextTimeStep, LinearLimits, PrevBody, ShapeTransform.GetTranslation(), &RigidBody, Body1JointOffset);
 		}
 		else
 		{
-			if (ConstraintSetup.LinearXLimitType != AnimPhysLinearConstraintType::Free)
+			if (PhysicsBodyDef.ConstraintSetup.LinearXLimitType != AnimPhysLinearConstraintType::Free)
 			{
-				FAnimPhys::ConstrainAlongDirection(NextTimeStep, LinearLimits, PrevBody, ShapeTransform.GetTranslation(), &RigidBody, Body1JointOffset, ShapeTransform.GetRotation().GetAxisX(), FVector2D(ConstraintSetup.LinearAxesMin.X, ConstraintSetup.LinearAxesMax.X));
+				FAnimPhys::ConstrainAlongDirection(NextTimeStep, LinearLimits, PrevBody, ShapeTransform.GetTranslation(), &RigidBody, Body1JointOffset, ShapeTransform.GetRotation().GetAxisX(), FVector2D(PhysicsBodyDef.ConstraintSetup.LinearAxesMin.X, PhysicsBodyDef.ConstraintSetup.LinearAxesMax.X));
 			}
 
-			if (ConstraintSetup.LinearYLimitType != AnimPhysLinearConstraintType::Free)
+			if (PhysicsBodyDef.ConstraintSetup.LinearYLimitType != AnimPhysLinearConstraintType::Free)
 			{
-				FAnimPhys::ConstrainAlongDirection(NextTimeStep, LinearLimits, PrevBody, ShapeTransform.GetTranslation(), &RigidBody, Body1JointOffset, ShapeTransform.GetRotation().GetAxisY(), FVector2D(ConstraintSetup.LinearAxesMin.Y, ConstraintSetup.LinearAxesMax.Y));
+				FAnimPhys::ConstrainAlongDirection(NextTimeStep, LinearLimits, PrevBody, ShapeTransform.GetTranslation(), &RigidBody, Body1JointOffset, ShapeTransform.GetRotation().GetAxisY(), FVector2D(PhysicsBodyDef.ConstraintSetup.LinearAxesMin.Y, PhysicsBodyDef.ConstraintSetup.LinearAxesMax.Y));
 			}
 
-			if (ConstraintSetup.LinearZLimitType != AnimPhysLinearConstraintType::Free)
+			if (PhysicsBodyDef.ConstraintSetup.LinearZLimitType != AnimPhysLinearConstraintType::Free)
 			{
-				FAnimPhys::ConstrainAlongDirection(NextTimeStep, LinearLimits, PrevBody, ShapeTransform.GetTranslation(), &RigidBody, Body1JointOffset, ShapeTransform.GetRotation().GetAxisZ(), FVector2D(ConstraintSetup.LinearAxesMin.Z, ConstraintSetup.LinearAxesMax.Z));
+				FAnimPhys::ConstrainAlongDirection(NextTimeStep, LinearLimits, PrevBody, ShapeTransform.GetTranslation(), &RigidBody, Body1JointOffset, ShapeTransform.GetRotation().GetAxisZ(), FVector2D(PhysicsBodyDef.ConstraintSetup.LinearAxesMin.Z, PhysicsBodyDef.ConstraintSetup.LinearAxesMax.Z));
 			}
 		}
 
-		if (ConstraintSetup.AngularConstraintType == AnimPhysAngularConstraintType::Angular)
+		if (PhysicsBodyDef.ConstraintSetup.AngularConstraintType == AnimPhysAngularConstraintType::Angular)
 		{
 #if WITH_EDITOR
 			// Check the ranges are valid when running in the editor, log if something is wrong
-			if(ConstraintSetup.AngularLimitsMin.X > ConstraintSetup.AngularLimitsMax.X ||
-			   ConstraintSetup.AngularLimitsMin.Y > ConstraintSetup.AngularLimitsMax.Y ||
-			   ConstraintSetup.AngularLimitsMin.Z > ConstraintSetup.AngularLimitsMax.Z)
+			if(PhysicsBodyDef.ConstraintSetup.AngularLimitsMin.X > PhysicsBodyDef.ConstraintSetup.AngularLimitsMax.X ||
+			   PhysicsBodyDef.ConstraintSetup.AngularLimitsMin.Y > PhysicsBodyDef.ConstraintSetup.AngularLimitsMax.Y ||
+			   PhysicsBodyDef.ConstraintSetup.AngularLimitsMin.Z > PhysicsBodyDef.ConstraintSetup.AngularLimitsMax.Z)
 			{
 				UE_LOG(LogAnimation, Warning, TEXT("AnimDynamics: Min/Max angular limits for bone %s incorrect, at least one min axis value is greater than the corresponding max."), *BoundBone.BoneName.ToString());
 			}
 #endif
 
 			// Add angular limits. any limit with 360+ degree range is ignored and left free.
-			FAnimPhys::ConstrainAngularRange(NextTimeStep, AngularLimits, PrevBody, &RigidBody, ShapeTransform.GetRotation(), ConstraintSetup.TwistAxis, ConstraintSetup.AngularLimitsMin, ConstraintSetup.AngularLimitsMax, bOverrideAngularBias ? AngularBiasOverride : AnimPhysicsConstants::JointBiasFactor);
+			FAnimPhys::ConstrainAngularRange(NextTimeStep, AngularLimits, PrevBody, &RigidBody, ShapeTransform.GetRotation(), PhysicsBodyDef.ConstraintSetup.TwistAxis, PhysicsBodyDef.ConstraintSetup.AngularLimitsMin, PhysicsBodyDef.ConstraintSetup.AngularLimitsMax, bOverrideAngularBias ? AngularBiasOverride : AnimPhysicsConstants::JointBiasFactor);
 		}
 		else
 		{
-			FAnimPhys::ConstrainConeAngle(NextTimeStep, AngularLimits, PrevBody, BoundBoneTransform.GetRotation().GetAxisX(), &RigidBody, FVector(1.0f, 0.0f, 0.0f), ConstraintSetup.ConeAngle, bOverrideAngularBias ? AngularBiasOverride : AnimPhysicsConstants::JointBiasFactor);
+			FAnimPhys::ConstrainConeAngle(NextTimeStep, AngularLimits, PrevBody, BoundBoneTransform.GetRotation().GetAxisX(), &RigidBody, FVector(1.0f, 0.0f, 0.0f), PhysicsBodyDef.ConstraintSetup.ConeAngle, bOverrideAngularBias ? AngularBiasOverride : AnimPhysicsConstants::JointBiasFactor);
 		}
 
 		if(PlanarLimits.Num() > 0 && bUsePlanarLimit)
@@ -932,8 +981,8 @@ void FAnimNode_AnimDynamics::UpdateLimits(FComponentSpacePoseContext& Output)
 			FAnimPhysSpring& NewSpring = Springs.Last();
 			NewSpring.SpringConstantLinear = LinearSpringConstant;
 			NewSpring.SpringConstantAngular = AngularSpringConstant;
-			NewSpring.AngularTarget = ConstraintSetup.AngularTarget.GetSafeNormal();
-			NewSpring.AngularTargetAxis = ConstraintSetup.AngularTargetAxis;
+			NewSpring.AngularTarget = PhysicsBodyDef.ConstraintSetup.AngularTarget.GetSafeNormal();
+			NewSpring.AngularTargetAxis = PhysicsBodyDef.ConstraintSetup.AngularTargetAxis;
 			NewSpring.TargetOrientationOffset = ShapeTransform.GetRotation();
 			NewSpring.bApplyAngular = bAngularSpring;
 			NewSpring.bApplyLinear = bLinearSpring;
@@ -1048,19 +1097,19 @@ int32 FAnimNode_AnimDynamics::GetLODThreshold() const
 	}
 }
 
-FTransform FAnimNode_AnimDynamics::GetBoneTransformInSimSpace(FComponentSpacePoseContext& Output, const FCompactPoseBoneIndex& BoneIndex)
+FTransform FAnimNode_AnimDynamics::GetBoneTransformInSimSpace(FComponentSpacePoseContext& Output, const FCompactPoseBoneIndex& BoneIndex) const
 {
 	FTransform Transform = Output.Pose.GetComponentSpaceTransform(BoneIndex);
 
 	return GetSimSpaceTransformFromComponentSpace(SimulationSpace, Output, Transform);
 }
 
-FTransform FAnimNode_AnimDynamics::GetComponentSpaceTransformFromSimSpace(AnimPhysSimSpaceType SimSpace, FComponentSpacePoseContext& Output, const FTransform& InSimTransform)
+FTransform FAnimNode_AnimDynamics::GetComponentSpaceTransformFromSimSpace(AnimPhysSimSpaceType SimSpace, FComponentSpacePoseContext& Output, const FTransform& InSimTransform) const
 {
 	return GetComponentSpaceTransformFromSimSpace(SimSpace, Output, InSimTransform, Output.AnimInstanceProxy->GetComponentTransform(), Output.AnimInstanceProxy->GetActorTransform());
 }
 
-FTransform FAnimNode_AnimDynamics::GetComponentSpaceTransformFromSimSpace(AnimPhysSimSpaceType SimSpace, FComponentSpacePoseContext& Output, const FTransform& InSimTransform, const FTransform& InCompWorldSpaceTM, const FTransform& InActorWorldSpaceTM)
+FTransform FAnimNode_AnimDynamics::GetComponentSpaceTransformFromSimSpace(AnimPhysSimSpaceType SimSpace, FComponentSpacePoseContext& Output, const FTransform& InSimTransform, const FTransform& InCompWorldSpaceTM, const FTransform& InActorWorldSpaceTM) const
 {
 	FTransform OutTransform = InSimTransform;
 
@@ -1116,8 +1165,29 @@ FTransform FAnimNode_AnimDynamics::GetComponentSpaceTransformFromSimSpace(AnimPh
 	return OutTransform;
 }
 
+FTransform FAnimNode_AnimDynamics::GetComponentSpaceTransformFromSimSpace(AnimPhysSimSpaceType SimSpace, const USkeletalMeshComponent* const SkelComp, const FTransform& InSimTransform) const
+{
+	FTransform OutTransformCS = InSimTransform;
+
+	check(SkelComp);
+	if (SkelComp)
+	{
+		if (SimSpace == AnimPhysSimSpaceType::RootRelative)
+		{
+			const FTransform RelativeBoneTransform = SkelComp->GetBoneTransform(0);
+			OutTransformCS = InSimTransform * RelativeBoneTransform;
+		}
+		else if (SimSpace == AnimPhysSimSpaceType::BoneRelative)
+		{
+			const FTransform RelativeBoneTransform = SkelComp->GetBoneTransform(SkelComp->GetBoneIndex(RelativeSpaceBone.BoneName));
+			OutTransformCS = InSimTransform * RelativeBoneTransform;
+		}
+	}
+
+	return OutTransformCS;
+}
 
-FTransform FAnimNode_AnimDynamics::GetSimSpaceTransformFromComponentSpace(AnimPhysSimSpaceType SimSpace, FComponentSpacePoseContext& Output, const FTransform& InComponentTransform)
+FTransform FAnimNode_AnimDynamics::GetSimSpaceTransformFromComponentSpace(AnimPhysSimSpaceType SimSpace, FComponentSpacePoseContext& Output, const FTransform& InComponentTransform) const
 {
 	FTransform ResultTransform = InComponentTransform;
 
@@ -1175,7 +1245,7 @@ FTransform FAnimNode_AnimDynamics::GetSimSpaceTransformFromComponentSpace(AnimPh
 	return ResultTransform;
 }
 
-FVector FAnimNode_AnimDynamics::TransformWorldVectorToSimSpace(FComponentSpacePoseContext& Output, const FVector& InVec)
+FVector FAnimNode_AnimDynamics::TransformWorldVectorToSimSpace(FComponentSpacePoseContext& Output, const FVector& InVec) const
 {
 	FVector OutVec = InVec;
 
@@ -1232,7 +1302,7 @@ FVector FAnimNode_AnimDynamics::TransformWorldVectorToSimSpace(FComponentSpacePo
 	return OutVec;
 }
 
-void FAnimNode_AnimDynamics::ConvertSimulationSpace(FComponentSpacePoseContext& Output, AnimPhysSimSpaceType From, AnimPhysSimSpaceType To)
+void FAnimNode_AnimDynamics::ConvertSimulationSpace(FComponentSpacePoseContext& Output, AnimPhysSimSpaceType From, AnimPhysSimSpaceType To) const
 {
 	for(FAnimPhysRigidBody* Body : BaseBodyPtrs)
 	{
@@ -1252,4 +1322,16 @@ void FAnimNode_AnimDynamics::ConvertSimulationSpace(FComponentSpacePoseContext&
 		Body->Pose.Orientation = BodyTransform.GetRotation();
 		Body->Pose.Position = BodyTransform.GetTranslation();
 	}
-}
\ No newline at end of file
+}
+
+bool FAnimNode_AnimDynamics::ShouldDoPhysicsUpdate() const
+{
+	#if WITH_EDITOR
+	if (!bDoPhysicsUpdateInEditor)
+	{
+		return false;
+	}
+	#endif
+
+	return bDoUpdate;
+}
diff --git a/Engine/Source/Runtime/AnimGraphRuntime/Public/BoneControllers/AnimNode_AnimDynamics.h b/Engine/Source/Runtime/AnimGraphRuntime/Public/BoneControllers/AnimNode_AnimDynamics.h
--- a/Engine/Source/Runtime/AnimGraphRuntime/Public/BoneControllers/AnimNode_AnimDynamics.h
+++ b/Engine/Source/Runtime/AnimGraphRuntime/Public/BoneControllers/AnimNode_AnimDynamics.h
@@ -234,6 +234,54 @@ struct FAnimPhysSphericalLimit
 	ESphericalLimitType LimitType;
 };
 
+USTRUCT()
+struct FAnimPhysBodyDefinition
+{
+	GENERATED_BODY();
+
+	FAnimPhysBodyDefinition()
+	: BoxExtents(10.0f, 10.0f, 10.0f)
+	, LocalJointOffset(FVector::ZeroVector)
+	, CollisionType(AnimPhysCollisionType::CoM)
+	, SphereCollisionRadius(10.0f)
+	{}
+
+	UPROPERTY(VisibleAnywhere, Category = PhysicsBodyDefinition, meta = (EditCondition = "false"))
+	FBoneReference BoundBone;
+
+	/** Extents of the box to use for simulation */
+	UPROPERTY(EditAnywhere, Category = PhysicsBodyDefinition, meta = (UIMin = "1", ClampMin = "1"))
+	FVector BoxExtents;
+
+	/** Vector relative to the body being simulated to attach the constraint to */
+	UPROPERTY(EditAnywhere, Category = PhysicsBodyDefinition)
+	FVector LocalJointOffset;
+
+	/** Data describing the constraints we will apply to the body */
+	UPROPERTY(EditAnywhere, Category = Constraint)
+	FAnimPhysConstraintSetup ConstraintSetup;
+
+	/** Resolution method for planar limits */
+	UPROPERTY(EditAnywhere, Category = Collision)
+	AnimPhysCollisionType CollisionType;
+
+	/** Radius to use if CollisionType is set to CustomSphere */
+	UPROPERTY(EditAnywhere, Category = PhysicsBodyDefinition, meta = (UIMin = "1", ClampMin = "1", EditCondition = "CollisionType == AnimPhysCollisionType::CustomSphere"))
+	float SphereCollisionRadius;
+};
+
+struct FAnimConstraintOffsetPair
+{
+	FAnimConstraintOffsetPair(const FVector& InBody0Offset, const FVector InBody1Offset)
+	: Body0Offset(InBody0Offset)
+	, Body1Offset(InBody1Offset)
+	{}
+
+	FVector Body0Offset;
+	FVector Body1Offset;
+};
+
+
 USTRUCT(BlueprintInternalUseOnly)
 struct ANIMGRAPHRUNTIME_API FAnimNode_AnimDynamics : public FAnimNode_SkeletalControlBase
 {
@@ -241,16 +289,16 @@ struct ANIMGRAPHRUNTIME_API FAnimNode_AnimDynamics : public FAnimNode_SkeletalCo
 
 	FAnimNode_AnimDynamics();
 
-	/** 
-	 * Overridden linear damping value. The default is 0.7. Values below 0.7 won't have an effect.
-	 */
-	UPROPERTY(EditAnywhere, AdvancedDisplay, BlueprintReadWrite, Category = Setup, meta = (PinHiddenByDefault, DisplayAfter="bOverrideLinearDamping"))
+	/**
+	* Overridden linear damping value. The default is 0.7. Values below 0.7 won't have an effect.
+	*/
+	UPROPERTY(EditAnywhere, AdvancedDisplay, BlueprintReadWrite, Category = Setup, meta = (PinHiddenByDefault, DisplayAfter = "NumSolverIterationsPostUpdate", EditCondition = bOverrideLinearDamping))
 	float LinearDampingOverride;
 
-	/** 
+	/**
 	 * Overridden angular damping value. The default is 0.7. Values below 0.7 won't have an effect.
 	 */
-	UPROPERTY(EditAnywhere, AdvancedDisplay, BlueprintReadWrite, Category = Setup, meta = (PinHiddenByDefault, DisplayAfter="bOverrideAngularDamping"))
+	UPROPERTY(EditAnywhere, AdvancedDisplay, BlueprintReadWrite, Category = Setup, meta = (PinHiddenByDefault, DisplayAfter = "LinearDampingOverride", EditCondition = bOverrideAngularDamping))
 	float AngularDampingOverride;
 
 	// Previous component & actor transforms, used to account for teleports
@@ -258,24 +306,19 @@ struct ANIMGRAPHRUNTIME_API FAnimNode_AnimDynamics : public FAnimNode_SkeletalCo
 	FTransform PreviousActorWorldSpaceTM;
 
 	/** When in BoneRelative sim space, the simulation will use this bone as the origin */
-	UPROPERTY(EditAnywhere, Category = Setup, meta=(DisplayAfter="SimulationSpace"))
+	UPROPERTY(EditAnywhere, Category = Setup, meta=(DisplayAfter="SimulationSpace", EditCondition = "SimulationSpace == AnimPhysSimSpaceType::BoneRelative"))
 	FBoneReference RelativeSpaceBone;
 
 	/** The bone to attach the physics body to, if bChain is true this is the top of the chain */
-	UPROPERTY(EditAnywhere, Category = Setup, meta=(DisplayAfter="bChain"))
+	UPROPERTY(EditAnywhere, Category = Setup)
 	FBoneReference BoundBone;
-	
+
 	/** If bChain is true this is the bottom of the chain, otherwise ignored */
-	UPROPERTY(EditAnywhere, Category = Setup, meta=(EditCondition = bChain))
+	UPROPERTY(EditAnywhere, Category = Setup, meta = (EditCondition = bChain, DisplayAfter = "BoundBone"))
 	FBoneReference ChainEnd;
 
-	/** Extents of the box to use for simulation */
-	UPROPERTY(EditAnywhere, Category = Setup)
-	FVector BoxExtents;
-
-	/** Vector relative to the body being simulated to attach the constraint to */
-	UPROPERTY(EditAnywhere, Category = Setup)
-	FVector LocalJointOffset;
+	UPROPERTY(EditAnywhere, EditFixedSize, Category = Setup, meta = (DisplayName = "Physics Bodies", EditFixedOrder, DisplayAfter = "ChainEnd"))
+	TArray< FAnimPhysBodyDefinition > PhysicsBodyDefinitions;
 
 	/** Scale for gravity, higher values increase forces due to gravity */
 	UPROPERTY(EditAnywhere, BlueprintReadWrite, Category = Setup, meta = (PinHiddenByDefault, DisplayAfter = "bGravityOverrideInSimSpace", EditCondition = "!bUseGravityOverride"))
@@ -289,7 +332,7 @@ struct ANIMGRAPHRUNTIME_API FAnimNode_AnimDynamics : public FAnimNode_SkeletalCo
 	 * Spring constant to use when calculating linear springs, higher values mean a stronger spring.
 	 * You need to enable the Linear Spring checkbox for this to have an effect.
 	 */
-	UPROPERTY(EditAnywhere, BlueprintReadWrite, Category = Setup, meta = (PinHiddenByDefault, DisplayAfter="bAngularSpring", EditCondition = "bLinearSpring"))
+	UPROPERTY(EditAnywhere, BlueprintReadWrite, Category = Setup, meta = (EditCondition = bLinearSpring, PinHiddenByDefault, DisplayAfter = "AngularSpringConstant"))
 	float LinearSpringConstant;
 
 	/** 
@@ -297,7 +340,7 @@ struct ANIMGRAPHRUNTIME_API FAnimNode_AnimDynamics : public FAnimNode_SkeletalCo
 	 * You need to enable the Angular Spring checkbox for this to have an effect.
 	 * Note: Make sure to also set the Angular Target Axis and Angular Target in the Constraint Setup for this to have an effect.
 	 */
-	UPROPERTY(EditAnywhere, BlueprintReadWrite, Category = Setup, meta = (PinHiddenByDefault, EditCondition = "bAngularSpring"))
+	UPROPERTY(EditAnywhere, BlueprintReadWrite, Category = Setup, meta = (DisplayAfter = "PhysicsBodyDefinitions", EditCondition = bAngularSpring, PinHiddenByDefault))
 	float AngularSpringConstant;
 
 	/** Scale to apply to calculated wind velocities in the solver */
@@ -318,32 +361,24 @@ struct ANIMGRAPHRUNTIME_API FAnimNode_AnimDynamics : public FAnimNode_SkeletalCo
 
 	/** Overridden angular bias value
 	 *  Angular bias is essentially a twist reduction for chain forces and defaults to a value to keep chains stability
-	 *  in check. When using single-body systems sometimes angular forces will look like they are "catching-up" with
-	 *  the mesh, if that's the case override this and push it towards 1.0f until it settles correctly
-	 */
-	UPROPERTY(EditAnywhere, AdvancedDisplay, BlueprintReadWrite, Category = Setup, meta = (PinHiddenByDefault, DisplayAfter="bOverrideAngularBias"))
+	*  in check. When using single-body systems sometimes angular forces will look like they are "catching-up" with
+	*  the mesh, if that's the case override this and push it towards 1.0f until it settles correctly
+	*/
+	UPROPERTY(EditAnywhere, AdvancedDisplay, BlueprintReadWrite, Category = Setup, meta = (PinHiddenByDefault, DisplayAfter = "AngularDampingOverride", EditCondition = bOverrideAngularBias))
 	float AngularBiasOverride;
 
 	/** Number of update passes on the linear and angular limits before we solve the position of the bodies recommended to be four times the value of NumSolverIterationsPostUpdate */
-	UPROPERTY(EditAnywhere, AdvancedDisplay, Category = Setup, meta=(DisplayAfter="bDoEval"))
+	UPROPERTY(EditAnywhere, AdvancedDisplay, Category = Setup)
 	int32 NumSolverIterationsPreUpdate;
 
 	/** Number of update passes on the linear and angular limits after we solve the position of the bodies, recommended to be around a quarter of NumSolverIterationsPreUpdate */
-	UPROPERTY(EditAnywhere, AdvancedDisplay, Category = Setup)
+	UPROPERTY(EditAnywhere, AdvancedDisplay, Category = Setup, meta = (PinHiddenByDefault, DisplayAfter = "NumSolverIterationsPreUpdate"))
 	int32 NumSolverIterationsPostUpdate;
 
-	/** Data describing the constraints we will apply to the body */
-	UPROPERTY(EditAnywhere, Category = Constraint)
-	FAnimPhysConstraintSetup ConstraintSetup;
-
 	/** List of available spherical limits for this node */
 	UPROPERTY(EditAnywhere, Category = SphericalLimit, meta=(DisplayAfter="bUseSphericalLimits"))
 	TArray<FAnimPhysSphericalLimit> SphericalLimits;
 
-	/** Radius to use if CollisionType is set to CustomSphere */
-	UPROPERTY(EditAnywhere, Category = Collision, meta = (UIMin = "1", ClampMin = "1", DisplayAfter="CollisionType"))
-	float SphereCollisionRadius;
-
 	/** An external force to apply to all bodies in the simulation when ticked, specified in world space */
 	UPROPERTY(EditAnywhere, Category = Forces, meta = (PinShownByDefault))
 	FVector ExternalForce;
@@ -352,10 +387,6 @@ struct ANIMGRAPHRUNTIME_API FAnimNode_AnimDynamics : public FAnimNode_SkeletalCo
 	UPROPERTY(EditAnywhere, Category=PlanarLimit, meta=(DisplayAfter="bUsePlanarLimit"))
 	TArray<FAnimPhysPlanarLimit> PlanarLimits;
 
-	/** Resolution method for planar limits */
-	UPROPERTY(EditAnywhere, Category = Collision)
-	AnimPhysCollisionType CollisionType;
-
 	/** The space used to run the simulation */
 	UPROPERTY(EditAnywhere, BlueprintReadWrite, Category = Setup, meta = (PinHiddenByDefault, DisplayPriority=0))
 	AnimPhysSimSpaceType SimulationSpace;
@@ -375,16 +406,16 @@ struct ANIMGRAPHRUNTIME_API FAnimNode_AnimDynamics : public FAnimNode_SkeletalCo
 	UPROPERTY(EditAnywhere, Category=PlanarLimit)
 	uint8 bUsePlanarLimit:1;
 
-	/** If true we will perform physics update, otherwise skip - allows visualisation of the initial state of the bodies */
-	UPROPERTY(EditAnywhere, AdvancedDisplay, Category = Setup, meta=(DisplayAfter="AngularBiasOverride"))
-	uint8 bDoUpdate:1;
+	/** If true we will perform physics update, otherwise skip - allows visualization of the initial state of the bodies */
+	UPROPERTY(EditAnywhere, AdvancedDisplay, Category = Setup, meta = (DisplayAfter = "bDoEval"))
+	uint8 bDoUpdate : 1;
 
-	/** If true we will perform bone transform evaluation, otherwise skip - allows visualisation of the initial anim state compared to the physics sim */
-	UPROPERTY(EditAnywhere, AdvancedDisplay, Category = Setup)
-	uint8 bDoEval:1;
+	/** If true we will perform bone transform evaluation, otherwise skip - allows visualization of the initial anim state compared to the physics sim */
+	UPROPERTY(EditAnywhere, AdvancedDisplay, Category = Setup, meta = (DisplayAfter = "AngularBiasOverride"))
+	uint8 bDoEval : 1;
 
 	/** If true, the override value will be used for linear damping */
-	UPROPERTY(EditAnywhere, AdvancedDisplay, Category = Setup, meta=(DisplayAfter="AngularSpringConstraint"))
+	UPROPERTY(EditAnywhere, AdvancedDisplay, Category = Setup, meta=(InlineEditConditionToggle, DisplayAfter="AngularSpringConstraint"))
 	uint8 bOverrideLinearDamping:1;
 
 	/** If true, the override value will be used for the angular bias for bodies in this node. 
@@ -392,11 +423,11 @@ struct ANIMGRAPHRUNTIME_API FAnimNode_AnimDynamics : public FAnimNode_SkeletalCo
 	 *  in check. When using single-body systems sometimes angular forces will look like they are "catching-up" with
 	 *  the mesh, if that's the case override this and push it towards 1.0f until it settles correctly
 	 */
-	UPROPERTY(EditAnywhere, AdvancedDisplay, Category = Setup, meta=(DisplayAfter="AngularDampingOverride"))
+	UPROPERTY(EditAnywhere, AdvancedDisplay, Category = Setup, meta=(InlineEditConditionToggle, DisplayAfter="AngularDampingOverride"))
 	uint8 bOverrideAngularBias:1;
 
 	/** If true, the override value will be used for angular damping */
-	UPROPERTY(EditAnywhere, AdvancedDisplay, Category = Setup, meta=(DisplayAfter="LinearDampingOverride"))
+	UPROPERTY(EditAnywhere, AdvancedDisplay, Category = Setup, meta=(InlineEditConditionToggle, DisplayAfter="LinearDampingOverride"))
 	uint8 bOverrideAngularDamping:1;
 
 	/** Whether or not wind is enabled for the bodies in this simulation */
@@ -406,7 +437,7 @@ struct ANIMGRAPHRUNTIME_API FAnimNode_AnimDynamics : public FAnimNode_SkeletalCo
 	uint8 bWindWasEnabled:1;
 
 	/** Use gravity override value vs gravity scale */
-	UPROPERTY(EditAnywhere, BlueprintReadWrite, Category = Setup)
+	UPROPERTY(EditAnywhere, BlueprintReadWrite, Category = Setup, meta = (DisplayAfter = "LinearSpringConstant"))
 	uint8 bUseGravityOverride:1;
 
 	/** If true the gravity override value is defined in simulation space, by default it is in world space */
@@ -414,21 +445,34 @@ struct ANIMGRAPHRUNTIME_API FAnimNode_AnimDynamics : public FAnimNode_SkeletalCo
 	uint8 bGravityOverrideInSimSpace : 1;
 
 	/** If true the body will attempt to spring back to its initial position */
-	UPROPERTY(EditAnywhere, Category = Setup, meta=(DisplayAfter="GravityScale"))
+	UPROPERTY(EditAnywhere, Category = Setup, meta = (InlineEditConditionToggle))
 	uint8 bLinearSpring:1;
 
 	/** If true the body will attempt to align itself with the specified angular target */
-	UPROPERTY(EditAnywhere, Category = Setup, meta=(DisplayAfter="bLinearSpring"))
+	UPROPERTY(EditAnywhere, Category = Setup, meta = (InlineEditConditionToggle))
 	uint8 bAngularSpring:1;
 
 	/** Set to true to use the solver to simulate a connected chain */
-	UPROPERTY(EditAnywhere, Category = Setup, meta=(DisplayAfter="RelativeSpaceBone"))
+	UPROPERTY(EditAnywhere, Category = Setup, meta=(InlineEditConditionToggle))
 	uint8 bChain:1;
 
 	/** The settings for rotation retargeting */
 	UPROPERTY(EditAnywhere, Category = Retargeting)
 	FRotationRetargetingInfo RetargetingSettings;
 
+#if WITH_EDITORONLY_DATA
+	UPROPERTY()
+		FVector BoxExtents_DEPRECATED;
+	UPROPERTY()
+		FVector LocalJointOffset_DEPRECATED;
+	UPROPERTY()
+		FAnimPhysConstraintSetup ConstraintSetup_DEPRECATED;
+	UPROPERTY()
+		AnimPhysCollisionType CollisionType_DEPRECATED;
+	UPROPERTY()
+		float SphereCollisionRadius_DEPRECATED;
+#endif
+
 	// FAnimNode_SkeletalControlBase interface
 	virtual void Initialize_AnyThread(const FAnimationInitializeContext& Context) override;
 	virtual void UpdateInternal(const FAnimationUpdateContext& Context) override;
@@ -445,22 +489,28 @@ struct ANIMGRAPHRUNTIME_API FAnimNode_AnimDynamics : public FAnimNode_SkeletalCo
 	void InitPhysics(FComponentSpacePoseContext& Output);
 	void TermPhysics();
 
+	void UpdateChainPhysicsBodyDefinitions(const USkeletalMeshComponent* const PreviewSkelMeshComp);
+	void UpdateChainPhysicsBodyDefinitions(const FReferenceSkeleton& ReferenceSkeleton);
+	void UpdateChainPhysicsBodyDefinitions(const TArray<FName>& ChainBoneNames);
 	void UpdateLimits(FComponentSpacePoseContext& Output);
 
 	int32 GetNumBodies() const;
 	const FAnimPhysRigidBody& GetPhysBody(int32 BodyIndex) const;
 
+	FTransform GetBodyComponentSpaceTransform(const FAnimPhysRigidBody& Body, const USkeletalMeshComponent* const SkelComp) const;
+
 #if WITH_EDITOR
 
 	// Accessors for editor code (mainly for visualization functions)
-	FVector GetBodyLocalJointOffset(int32 BodyIndex) const;
-	// Gets the number of currently bound bones (always 1 unless using a chain)
-	int32 GetNumBoundBones() const;
-	// Gets the specified bound bone reference
-	const FBoneReference* GetBoundBoneReference(int32 Index) const;
+	FVector GetBodyLocalJointOffset(const int32 BodyIndex) const;
+
+	// True by default, if false physics simulation will not update this frame. Used to prevent the rig moving whilst interactively editing parameters with a widget in the viewport.
+	bool bDoPhysicsUpdateInEditor;
 
 #endif
 
+	bool ShouldDoPhysicsUpdate() const;
+
 protected:
 
 	// FAnimNode_SkeletalControlBase protected interface
@@ -469,21 +519,21 @@ protected:
 	// End of FAnimNode_SkeletalControlBase protected interface
 
 private:
-
 	// Given a bone index, get it's transform in the currently selected simulation space
-	FTransform GetBoneTransformInSimSpace(FComponentSpacePoseContext& Output, const FCompactPoseBoneIndex& BoneIndex);
+	FTransform GetBoneTransformInSimSpace(FComponentSpacePoseContext& Output, const FCompactPoseBoneIndex& BoneIndex) const;
 
 	// Given a transform in simulation space, convert it back to component space
-	FTransform GetComponentSpaceTransformFromSimSpace(AnimPhysSimSpaceType SimSpace, FComponentSpacePoseContext& Output, const FTransform& InSimTransform);
-	FTransform GetComponentSpaceTransformFromSimSpace(AnimPhysSimSpaceType SimSpace, FComponentSpacePoseContext& Output, const FTransform& InSimTransform, const FTransform& InCompWorldSpaceTM, const FTransform& InActorWorldSpaceTM);
+	FTransform GetComponentSpaceTransformFromSimSpace(AnimPhysSimSpaceType SimSpace, FComponentSpacePoseContext& Output, const FTransform& InSimTransform) const;
+	FTransform GetComponentSpaceTransformFromSimSpace(AnimPhysSimSpaceType SimSpace, FComponentSpacePoseContext& Output, const FTransform& InSimTransform, const FTransform& InCompWorldSpaceTM, const FTransform& InActorWorldSpaceTM) const;
+	FTransform GetComponentSpaceTransformFromSimSpace(AnimPhysSimSpaceType SimSpace, const USkeletalMeshComponent* const SkelComp, const FTransform& InSimTransform) const;
 
 	// Given a transform in component space, convert it to the current sim space
-	FTransform GetSimSpaceTransformFromComponentSpace(AnimPhysSimSpaceType SimSpace, FComponentSpacePoseContext& Output, const FTransform& InComponentTransform);
+	FTransform GetSimSpaceTransformFromComponentSpace(AnimPhysSimSpaceType SimSpace, FComponentSpacePoseContext& Output, const FTransform& InComponentTransform) const;
 
 	// Given a world-space vector, convert it into the current simulation space
-	FVector TransformWorldVectorToSimSpace(FComponentSpacePoseContext& Output, const FVector& InVec);
+	FVector TransformWorldVectorToSimSpace(FComponentSpacePoseContext& Output, const FVector& InVec) const;
 
-	void ConvertSimulationSpace(FComponentSpacePoseContext& Output, AnimPhysSimSpaceType From, AnimPhysSimSpaceType To);
+	void ConvertSimulationSpace(FComponentSpacePoseContext& Output, AnimPhysSimSpaceType From, AnimPhysSimSpaceType To) const;
 
 	// Maximum time to consider when accumulating time debt to avoid spiraling
 	static const float MaxTimeDebt;
@@ -520,13 +570,13 @@ private:
 	// List of spring force generators created for this frame
 	TArray<FAnimPhysSpring> Springs;
 
-	// Local space offsets for each body
-	TArray<FVector> JointOffsets;
+	// Position of the physics object relative to the transform if its bound bone.
+	TArray<FVector> PhysicsBodyJointOffsets;
 
-	// List of bone references for all bodies in this node
-	TArray<FBoneReference> BoundBoneReferences;
-
-	// Depending on the LOD we might not be runnning all of the bound bodies (for chains)
+	// A pair of positions (relative to their associated physics bodies) for each pair of bodies in a chain. These positions should be driven to match each other in sim space by the physics contstraints - See UpdateLimits() fns.
+	TArray<FAnimConstraintOffsetPair> ConstraintOffsets;
+	
+	// Depending on the LOD we might not be running all of the bound bodies (for chains)
 	// this tracks the active bodies.
 	TArray<int32> ActiveBoneIndices;
 
diff --git a/Engine/Source/Runtime/Core/Public/UObject/FortniteMainBranchObjectVersion.h b/Engine/Source/Runtime/Core/Public/UObject/FortniteMainBranchObjectVersion.h
--- a/Engine/Source/Runtime/Core/Public/UObject/FortniteMainBranchObjectVersion.h
+++ b/Engine/Source/Runtime/Core/Public/UObject/FortniteMainBranchObjectVersion.h
@@ -190,6 +190,8 @@ struct CORE_API FFortniteMainBranchObjectVersion
 		// true for nodes predating this change.
 		GravityOverrideDefinedInWorldSpace,
 
+		// Anim Dynamics Node Physics parameters for each body in a chain are now stored in an array and can be edited.
+		AnimDynamicsEditableChainParameters,
 
 		// -----<new versions can be added above this line>-------------------------------------------------
 		VersionPlusOne,
diff --git a/Engine/Source/Runtime/Engine/Private/Animation/AnimPhysicsSolver.cpp b/Engine/Source/Runtime/Engine/Private/Animation/AnimPhysicsSolver.cpp
--- a/Engine/Source/Runtime/Engine/Private/Animation/AnimPhysicsSolver.cpp
+++ b/Engine/Source/Runtime/Engine/Private/Animation/AnimPhysicsSolver.cpp
@@ -59,22 +59,22 @@ FAnimPhysShape::FAnimPhysShape(TArray<FVector>& InVertices, TArray<FIntVector>&
 	CenterOfMass = FAnimPhys::CalculateCenterOfMass(Vertices, Triangles);
 }
 
-FAnimPhysShape FAnimPhysShape::MakeBox(FVector& Extents)
+FAnimPhysShape FAnimPhysShape::MakeBox(const FVector& Extents)
 {
+	FVector HalfExtents = Extents / 2.0f;
+
 	// A box of zero size will introduce NaNs into the simulation so we stomp it here and log
 	// if we encounter it.
-	if(Extents.SizeSquared() <= SMALL_NUMBER)
+	if (Extents.SizeSquared() <= SMALL_NUMBER)
 	{
 // 		UE_LOG(LogAnimation, Warning, TEXT("AnimDynamics: Attempted to create a simulation box with 0 volume, this introduces NaNs into the simulation. Adjusting box extents to (1.0f,1.0f,1.0f)"));
-		Extents = FVector(1.0f);
+		HalfExtents = FVector(0.5f);
 	}
 
 	TArray<FVector> Verts;
 	TArray<FIntVector> Tris;
 	Verts.Reserve(8);
-	Tris.Reserve(12);
-
-	FVector HalfExtents = Extents / 2.0f;
+	Tris.Reserve(12);	
 
 	// Front Verts
 	Verts.Add(FVector(-HalfExtents.X, -HalfExtents.Y, HalfExtents.Z));
diff --git a/Engine/Source/Runtime/Engine/Public/Animation/AnimPhysicsSolver.h b/Engine/Source/Runtime/Engine/Public/Animation/AnimPhysicsSolver.h
--- a/Engine/Source/Runtime/Engine/Public/Animation/AnimPhysicsSolver.h
+++ b/Engine/Source/Runtime/Engine/Public/Animation/AnimPhysicsSolver.h
@@ -83,7 +83,7 @@ struct ENGINE_API FAnimPhysShape
 	/** Makes a box with the given extents
 	 *  @param Extents Extents of the resulting box
 	 */
-	static FAnimPhysShape MakeBox(FVector& Extents);
+	static FAnimPhysShape MakeBox(const FVector& Extents);
 
 	FAnimPhysShape();
 	FAnimPhysShape(TArray<FVector>& InVertices, TArray<FIntVector>& InTriangles);
