[FIX] for AnimDynamics Cook error introduced in CL-18191484

#rmx

diff --git a/Engine/Source/Runtime/AnimGraphRuntime/Private/BoneControllers/AnimNode_AnimDynamics.cpp b/Engine/Source/Runtime/AnimGraphRuntime/Private/BoneControllers/AnimNode_AnimDynamics.cpp
--- a/Engine/Source/Runtime/AnimGraphRuntime/Private/BoneControllers/AnimNode_AnimDynamics.cpp
+++ b/Engine/Source/Runtime/AnimGraphRuntime/Private/BoneControllers/AnimNode_AnimDynamics.cpp
@@ -28,24 +28,27 @@ void FindChainBones(const FName BeginBoneName, const FName EndBoneName, TFunctio
 {
 	OutChainBoneNames.Add(BeginBoneName);
 
-	const uint32 OutputStartIndex = OutChainBoneNames.Num();
+	if (!EndBoneName.IsNone())
+	{
+		const uint32 OutputStartIndex = OutChainBoneNames.Num();
 
-	// Add the end of the chain. We have to walk from the bottom upwards to find a chain
-	// as walking downwards doesn't guarantee a single end point.
+		// Add the end of the chain. We have to walk from the bottom upwards to find a chain
+		// as walking downwards doesn't guarantee a single end point.
 
-	// Walk up the chain until we either find the top or hit the root bone
-	FName ParentBoneName = EndBoneName;
-	for (; !ParentBoneName.IsNone() && (ParentBoneName != BeginBoneName); ParentBoneName = GetParentBoneName(ParentBoneName))
-	{
-		OutChainBoneNames.Insert(ParentBoneName, OutputStartIndex);
-	}
+		// Walk up the chain until we either find the top or hit the root bone
+		FName ParentBoneName = EndBoneName;
+		for (; !ParentBoneName.IsNone() && (ParentBoneName != BeginBoneName); ParentBoneName = GetParentBoneName(ParentBoneName))
+		{
+			OutChainBoneNames.Insert(ParentBoneName, OutputStartIndex);
+		}
 
-	if (ParentBoneName != BeginBoneName)
-	{
-		UE_LOG(LogAnimation, Error, TEXT("AnimDynamics: Attempted to find bone chain starting at %s and ending at %s but failed."), *BeginBoneName.ToString(), *EndBoneName.ToString());
+		if (ParentBoneName != BeginBoneName)
+		{
+			UE_LOG(LogAnimation, Error, TEXT("AnimDynamics: Attempted to find bone chain starting at %s and ending at %s but failed."), *BeginBoneName.ToString(), *EndBoneName.ToString());
 
-		// Remove any accumulated chain bone names beyond the start bone from output.
-		OutChainBoneNames.RemoveAt(OutputStartIndex, OutChainBoneNames.Num() - OutputStartIndex);
+			// Remove any accumulated chain bone names beyond the start bone from output.
+			OutChainBoneNames.RemoveAt(OutputStartIndex, OutChainBoneNames.Num() - OutputStartIndex);
+		}
 	}
 }
 
