[FIX] Check for invalid bone indexs when updating chain in AnimNode_AnimDynamics.

#rmx

diff --git a/Engine/Source/Runtime/AnimGraphRuntime/Private/BoneControllers/AnimNode_AnimDynamics.cpp b/Engine/Source/Runtime/AnimGraphRuntime/Private/BoneControllers/AnimNode_AnimDynamics.cpp
--- a/Engine/Source/Runtime/AnimGraphRuntime/Private/BoneControllers/AnimNode_AnimDynamics.cpp
+++ b/Engine/Source/Runtime/AnimGraphRuntime/Private/BoneControllers/AnimNode_AnimDynamics.cpp
@@ -475,8 +475,18 @@ void FAnimNode_AnimDynamics::UpdateChainPhysicsBodyDefinitions(const FReferenceS
 	auto GetParentBoneNameFn = [&ReferenceSkeleton](const FName BoneName)
 	{
 		int32 BoneIndex = ReferenceSkeleton.FindBoneIndex(BoneName);
-		BoneIndex = ReferenceSkeleton.GetParentIndex(BoneIndex);
-		return ReferenceSkeleton.GetBoneName(BoneIndex);
+
+		if (ReferenceSkeleton.IsValidIndex(BoneIndex))
+		{
+			BoneIndex = ReferenceSkeleton.GetParentIndex(BoneIndex);
+		}
+
+		if (ReferenceSkeleton.IsValidIndex(BoneIndex))
+		{
+			return ReferenceSkeleton.GetBoneName(BoneIndex);
+		}
+
+		return FName(NAME_None);
 	};
 
 	TArray<FName> ChainBoneNames;
