Fixed sequence player node blueprint exposure issues

Fixed GetSequence to be a pure node, making it possible to use it in the anim graph
Additionally, marked UAnimSequenceBase::GetPlayLength BlueprintThreadSafe

diff --git a/Engine/Source/Runtime/AnimGraphRuntime/Private/SequencePlayerLibrary.cpp b/Engine/Source/Runtime/AnimGraphRuntime/Private/SequencePlayerLibrary.cpp
--- a/Engine/Source/Runtime/AnimGraphRuntime/Private/SequencePlayerLibrary.cpp
+++ b/Engine/Source/Runtime/AnimGraphRuntime/Private/SequencePlayerLibrary.cpp
@@ -114,6 +114,20 @@ FSequencePlayerReference USequencePlayerLibrary::GetSequence(const FSequencePlay
 	return SequencePlayer;
 }
 
+UAnimSequenceBase* USequencePlayerLibrary::GetSequencePure(const FSequencePlayerReference& SequencePlayer)
+{
+	UAnimSequenceBase* SequenceBase = nullptr;
+	
+	SequencePlayer.CallAnimNodeFunction<FAnimNode_SequencePlayer>(
+		TEXT("GetSequence"),
+		[&SequenceBase](FAnimNode_SequencePlayer& InSequencePlayer)
+		{
+			SequenceBase = InSequencePlayer.GetSequence();
+		});
+
+	return SequenceBase;
+}
+
 float USequencePlayerLibrary::GetAccumulatedTime(const FSequencePlayerReference& SequencePlayer)
 {
 	float AccumulatedTime = 0.f;
diff --git a/Engine/Source/Runtime/AnimGraphRuntime/Public/SequencePlayerLibrary.h b/Engine/Source/Runtime/AnimGraphRuntime/Public/SequencePlayerLibrary.h
--- a/Engine/Source/Runtime/AnimGraphRuntime/Public/SequencePlayerLibrary.h
+++ b/Engine/Source/Runtime/AnimGraphRuntime/Public/SequencePlayerLibrary.h
@@ -59,24 +59,28 @@ public:
 	/** Set the current sequence of the sequence player with an inertial blend time */
 	UFUNCTION(BlueprintCallable, Category = "Sequence Player", meta = (BlueprintThreadSafe))
 	static FSequencePlayerReference SetSequenceWithInertialBlending(const FAnimUpdateContext& UpdateContext, const FSequencePlayerReference& SequencePlayer, UAnimSequenceBase* Sequence, float BlendTime = 0.2f);
+
+	/** Get the current sequence of the sequence player - DEPRECATED, please use pure version */
+	UFUNCTION(BlueprintCallable, Category = "Sequence Player", meta = (BlueprintThreadSafe, DeprecatedFunction))
+	static FSequencePlayerReference GetSequence(const FSequencePlayerReference& SequencePlayer, UPARAM(Ref) UAnimSequenceBase*& SequenceBase);
 	
 	/** Get the current sequence of the sequence player */
-	UFUNCTION(BlueprintCallable, Category = "Sequence Player", meta = (BlueprintThreadSafe))
-	static FSequencePlayerReference GetSequence(const FSequencePlayerReference& SequencePlayer, UPARAM(Ref) UAnimSequenceBase*& SequenceBase);
+	UFUNCTION(BlueprintPure, Category = "Sequence Player", meta = (BlueprintThreadSafe, DisplayName = "Get Sequence"))
+	static UAnimSequenceBase* GetSequencePure(const FSequencePlayerReference& SequencePlayer);
 
 	/** Gets the current accumulated time of the sequence player */
-	UFUNCTION(BlueprintCallable, Category = "Sequence Player", meta = (BlueprintThreadSafe))
+	UFUNCTION(BlueprintPure, Category = "Sequence Player", meta = (BlueprintThreadSafe))
 	static float GetAccumulatedTime(const FSequencePlayerReference& SequencePlayer);
 
 	/** Get the start position of the sequence player */
-	UFUNCTION(BlueprintCallable, Category = "Sequence Player", meta = (BlueprintThreadSafe))
+	UFUNCTION(BlueprintPure, Category = "Sequence Player", meta = (BlueprintThreadSafe))
 	static float GetStartPosition(const FSequencePlayerReference& SequencePlayer);
 
 	/** Get the play rate of the sequence player */
-	UFUNCTION(BlueprintCallable, Category = "Sequence Player", meta = (BlueprintThreadSafe))
+	UFUNCTION(BlueprintPure, Category = "Sequence Player", meta = (BlueprintThreadSafe))
 	static float GetPlayRate(const FSequencePlayerReference& SequencePlayer);
 
 	/** Get the looping state of the sequence player */
-	UFUNCTION(BlueprintCallable, Category = "Sequence Player", meta = (BlueprintThreadSafe))
+	UFUNCTION(BlueprintPure, Category = "Sequence Player", meta = (BlueprintThreadSafe))
 	static bool GetLoopAnimation(const FSequencePlayerReference& SequencePlayer);
 };
\ No newline at end of file
diff --git a/Engine/Source/Runtime/Engine/Classes/Animation/AnimationAsset.h b/Engine/Source/Runtime/Engine/Classes/Animation/AnimationAsset.h
--- a/Engine/Source/Runtime/Engine/Classes/Animation/AnimationAsset.h
+++ b/Engine/Source/Runtime/Engine/Classes/Animation/AnimationAsset.h
@@ -993,7 +993,7 @@ public:
 	UE_DEPRECATED(5.0, "Use GetPlayLength instead")
 	virtual float GetMaxCurrentTime() { return GetPlayLength(); }
 
-	UFUNCTION(BlueprintCallable, Category = "Animation")
+	UFUNCTION(BlueprintPure, Category = "Animation", meta=(BlueprintThreadSafe))
 	virtual float GetPlayLength() const { return 0.f; };
 
 	void SetSkeleton(USkeleton* NewSkeleton);
