[UE][FIX] Bug fixes for Anim Dynamics UI
- Removed cached USkeletalMeshComponent* from AnimGraphNode as we can now get it from AnimGraphNode_Base
- Add a templatized version of AnimGraphNode_Base::GetDebuggedAnimNode() that allows retrival of derived node types
- Can now calculate chain bones without waiting for the node to be compiled - removed associated warning from UI
- Prevent Copy+Paste of bone names in chain bone array, improving workflow
- Ensure that changing bound bone name does not reset physics body parameters, improving workflow
- Improve chain body UI

diff --git a/Engine/Source/Editor/AnimGraph/Private/AnimGraphNode_AnimDynamics.cpp b/Engine/Source/Editor/AnimGraph/Private/AnimGraphNode_AnimDynamics.cpp
--- a/Engine/Source/Editor/AnimGraph/Private/AnimGraphNode_AnimDynamics.cpp
+++ b/Engine/Source/Editor/AnimGraph/Private/AnimGraphNode_AnimDynamics.cpp
@@ -3,21 +3,43 @@
 #include "AnimGraphNode_AnimDynamics.h"
 #include "AnimNodeEditModes.h"
 #include "EngineGlobals.h"
+#include "Kismet2/BlueprintEditorUtils.h"
 #include "SceneManagement.h"
 #include "Materials/MaterialInstanceDynamic.h"
 #include "UObject/ReleaseObjectVersion.h"
 #include "Widgets/Input/SButton.h"
 
 // Details includes
+#include "IDetailCustomization.h"
 #include "PropertyHandle.h"
 #include "DetailLayoutBuilder.h"
 #include "DetailWidgetRow.h"
 #include "DetailCategoryBuilder.h"
 #include "AnimationCustomVersion.h"
 #include "Animation/AnimInstance.h"
+#include "IDetailChildrenBuilder.h"
+#include "PropertyCustomizationHelpers.h"
 
 #define LOCTEXT_NAMESPACE "AnimDynamicsNode"
 
+////////////////////////////////////////
+// class FAnimGraphNode_AnimDynamics_DetailCustomization
+//
+// Required to customise the chain body array rendered in the details panel.
+//
+class FAnimGraphNode_AnimDynamics_DetailCustomization : public IDetailCustomization
+{
+public:
+	static TSharedRef< IDetailCustomization > MakeInstance();
+	virtual void CustomizeDetails(IDetailLayoutBuilder& DetailBuilder) override;
+	void OnPhysicsBodyDefCustomizeDetails(TSharedRef<class IPropertyHandle> ElementProperty, int32 ElementIndex, class IDetailChildrenBuilder& ChildrenBuilder, class IDetailLayoutBuilder* DetailLayout);
+
+	UClass* OwningClass;
+};
+
+////////////////////////////////////////
+// class UAnimGraphNode_AnimDynamics
+
 FText UAnimGraphNode_AnimDynamics::GetTooltipText() const
 {
 	return LOCTEXT("NodeTooltip", "Anim Dynamics");
@@ -68,32 +90,20 @@ void UAnimGraphNode_AnimDynamics::CustomizeDetails(IDetailLayoutBuilder& DetailB
 			.OnClicked(FOnClicked::CreateStatic(&UAnimGraphNode_AnimDynamics::ResetButtonClicked, &DetailBuilder))
 		];
 
-	// Add warning message about physics body array not being updated untill after the first BP compile.
-	if (LastPreviewComponent == nullptr)
+	if (!DetailCustomization)
 	{
-		IDetailCategoryBuilder& SetupCategory = DetailBuilder.EditCategory(TEXT("Setup"));
-		const FText WarningText(LOCTEXT("AnimDynamicsWarningText", "WARNING - Physics Bodies Will Not Be Valid Untill This Node Has Been Connected And Compiled"));
-		FDetailWidgetRow& WarningMessageTextWidgetRow = SetupCategory.AddCustomRow(WarningText);
-
-		WarningMessageTextWidgetRow
-			[
-				SNew(SHorizontalBox)
-				+ SHorizontalBox::Slot()
-				.AutoWidth()
-				[
-					SNew(STextBlock)
-					.Text(WarningText)
-					.Justification(ETextJustify::Center)
-					.ColorAndOpacity(FLinearColor::Red)
-				]
-			];
+		DetailCustomization = MakeShared<FAnimGraphNode_AnimDynamics_DetailCustomization>();
 	}
 
+	DetailCustomization->OwningClass = GetClass();
+	DetailCustomization->CustomizeDetails(DetailBuilder);
+
 	// Force order of details panel catagories - Must set order for all of them as any that are edited automatically move to the top.
 	{
 		uint32 SortOrder = 0;
 		DetailBuilder.EditCategory("Preview").SetSortOrder(SortOrder++);
 		DetailBuilder.EditCategory("Setup").SetSortOrder(SortOrder++);
+		DetailBuilder.EditCategory("PhysicsParameters").SetSortOrder(SortOrder++);
 		DetailBuilder.EditCategory("Settings").SetSortOrder(SortOrder++);
 		DetailBuilder.EditCategory("SphericalLimit").SetSortOrder(SortOrder++);
 		DetailBuilder.EditCategory("PlanarLimit").SetSortOrder(SortOrder++);
@@ -158,8 +168,23 @@ FText UAnimGraphNode_AnimDynamics::GetNodeTitle(ENodeTitleType::Type TitleType)
 	return CachedNodeTitles[TitleType];
 }
 
+
+USkeleton* UAnimGraphNode_AnimDynamics::GetSkeleton() const
+{
+	USkeleton* Skeleton = nullptr;
+
+	if (UAnimBlueprint* const AnimBlueprint = Cast<UAnimBlueprint>(FBlueprintEditorUtils::FindBlueprintForNode(this)))
+	{
+		Skeleton = AnimBlueprint->TargetSkeleton;
+	}
+
+	return Skeleton;
+}
+
 void UAnimGraphNode_AnimDynamics::PostEditChangeProperty(struct FPropertyChangedEvent& PropertyChangedEvent)
 {
+	bool bShouldRebuildChain = false;
+
 	if (PropertyChangedEvent.GetPropertyName() == GET_MEMBER_NAME_CHECKED(FAnimNode_AnimDynamics, ChainEnd))
 	{
 		// bChain has been modified.
@@ -172,13 +197,23 @@ void UAnimGraphNode_AnimDynamics::PostEditChangeProperty(struct FPropertyChanged
 			Node.ChainEnd.BoneName = NAME_None;
 		}
 
-		Node.UpdateChainPhysicsBodyDefinitions(LastPreviewComponent);
+		bShouldRebuildChain = true;
 	}
 
-	if (PropertyChangedEvent.GetPropertyName() == GET_MEMBER_NAME_CHECKED(FBoneReference, BoneName))
+	bShouldRebuildChain |= (PropertyChangedEvent.GetPropertyName() == GET_MEMBER_NAME_CHECKED(FBoneReference, BoneName)); // Either BoundBone or ChainEnd have been modified. 
+	
+	if (USkeleton* const Skeleton = GetSkeleton())
 	{
-		// Either BoundBone or ChainEnd have been modified.
-		Node.UpdateChainPhysicsBodyDefinitions(LastPreviewComponent);
+		if (bShouldRebuildChain)
+		{
+			// Rebuild Chain if begin or end bones have changed.
+			Node.UpdateChainPhysicsBodyDefinitions(Skeleton->GetReferenceSkeleton());
+		}
+		else
+		{
+			// Write chain bones names to each body in chain and check chain length in case either have been changed by a copy and paste operation.
+			Node.ValidateChainPhysicsBodyDefinitions(Skeleton->GetReferenceSkeleton());
+		}
 	}
 
 	Super::PostEditChangeProperty(PropertyChangedEvent);
@@ -200,18 +235,7 @@ void UAnimGraphNode_AnimDynamics::ResetSim()
 
 FAnimNode_AnimDynamics* UAnimGraphNode_AnimDynamics::GetPreviewDynamicsNode() const
 {
-	FAnimNode_AnimDynamics* ActivePreviewNode = nullptr;
-
-	if(LastPreviewComponent && LastPreviewComponent->GetAnimInstance())
-	{
-		UAnimInstance* Instance = LastPreviewComponent->GetAnimInstance();
-		if(UAnimBlueprintGeneratedClass* Class = Cast<UAnimBlueprintGeneratedClass>(Instance->GetClass()))
-		{
-			ActivePreviewNode = Class->GetPropertyInstance<FAnimNode_AnimDynamics>(Instance, NodeGuid);
-		}
-	}
-
-	return ActivePreviewNode;
+	return GetDebuggedAnimNode<FAnimNode_AnimDynamics>();
 }
 
 FReply UAnimGraphNode_AnimDynamics::ResetButtonClicked(IDetailLayoutBuilder* DetailLayoutBuilder)
@@ -265,4 +289,50 @@ void UAnimGraphNode_AnimDynamics::Serialize(FArchive& Ar)
 	}
 }
 
+TSharedRef< IDetailCustomization > FAnimGraphNode_AnimDynamics_DetailCustomization::MakeInstance()
+{
+	return MakeShareable(new FAnimGraphNode_AnimDynamics_DetailCustomization);
+}
+
+void FAnimGraphNode_AnimDynamics_DetailCustomization::CustomizeDetails(IDetailLayoutBuilder& DetailLayout)
+{
+	IDetailCategoryBuilder& SetupCategory = DetailLayout.EditCategory(TEXT("PhysicsParameters"));
+	TSharedRef< IPropertyHandle > PhysicsBodyDefinitionsProperty = DetailLayout.GetProperty("Node.PhysicsBodyDefinitions", OwningClass);
+	if (PhysicsBodyDefinitionsProperty->AsArray().IsValid())
+	{
+		TSharedRef<FDetailArrayBuilder> PropertyBuilder = MakeShareable(new FDetailArrayBuilder(PhysicsBodyDefinitionsProperty, /*InGenerateHeader*/ true, /*InDisplayResetToDefault*/ true, /*InDisplayElementNum*/ true));
+		PropertyBuilder->OnGenerateArrayElementWidget(FOnGenerateArrayElementWidget::CreateSP(this, &FAnimGraphNode_AnimDynamics_DetailCustomization::OnPhysicsBodyDefCustomizeDetails, &DetailLayout));
+		SetupCategory.AddCustomBuilder(PropertyBuilder, false);
+	}
+}
+
+void FAnimGraphNode_AnimDynamics_DetailCustomization::OnPhysicsBodyDefCustomizeDetails(TSharedRef<IPropertyHandle> ElementProperty, int32 ElementIndex, IDetailChildrenBuilder& ChildrenBuilder, IDetailLayoutBuilder* DetailLayout)
+{
+	TSharedPtr<IPropertyHandle> BoundBoneProperty = ElementProperty->GetChildHandle(GET_MEMBER_NAME_CHECKED(FAnimPhysBodyDefinition, BoundBone));
+
+	// Get Bone Name
+	FName BoneName;
+
+	if (BoundBoneProperty)
+	{
+		TSharedPtr<IPropertyHandle> BoundBoneNameProperty = BoundBoneProperty->GetChildHandle(GET_MEMBER_NAME_CHECKED(FBoneReference, BoneName));
+		if (BoundBoneNameProperty)
+		{
+			BoundBoneNameProperty->GetValue(BoneName);
+		}
+
+		DetailLayout->HideProperty(BoundBoneProperty);
+		IDetailPropertyRow& Row = ChildrenBuilder.AddProperty(ElementProperty);
+		
+		// Set a custom widget show a more useful item name and remove the 'n items' text that would otherwise appear on the body def group header.
+		const FString BodyDefNameString = "[" + FString::FromInt(ElementIndex) + "] " + BoneName.ToString();
+		Row.CustomWidget(true) 
+			.NameContent()
+			[
+				SNew(STextBlock)
+				.Text(FText::FromString(BodyDefNameString))
+			];
+	}
+}
+
 #undef LOCTEXT_NAMESPACE
diff --git a/Engine/Source/Editor/AnimGraph/Private/AnimGraphNode_Base.cpp b/Engine/Source/Editor/AnimGraph/Private/AnimGraphNode_Base.cpp
--- a/Engine/Source/Editor/AnimGraph/Private/AnimGraphNode_Base.cpp
+++ b/Engine/Source/Editor/AnimGraph/Private/AnimGraphNode_Base.cpp
@@ -1795,22 +1795,6 @@ bool UAnimGraphNode_Base::ReferencesFunction(const FName& InFunctionName,const U
 	return InitialUpdateFunction.GetMemberName() == InFunctionName || BecomeRelevantFunction.GetMemberName() == InFunctionName || UpdateFunction.GetMemberName() == InFunctionName;
 }
 
-FAnimNode_Base* UAnimGraphNode_Base::GetDebuggedAnimNode() const
-{
-	if (UBlueprint* Blueprint = FBlueprintEditorUtils::FindBlueprintForNode(this))
-	{
-		if (UObject* ActiveObject = Blueprint->GetObjectBeingDebugged())
-		{
-			if (UAnimBlueprintGeneratedClass* Class = Cast<UAnimBlueprintGeneratedClass>((UObject*)ActiveObject->GetClass()))
-			{
-				return Class->GetPropertyInstance<FAnimNode_Base>(ActiveObject, this);
-			}
-		}
-	}
-
-	return nullptr;
-}
-
 void UAnimGraphNode_Base::PostEditRefreshDebuggedComponent()
 {
 	if(FAnimNode_Base* DebuggedAnimNode = GetDebuggedAnimNode())
@@ -1834,6 +1818,16 @@ void UAnimGraphNode_Base::PostEditRefreshDebuggedComponent()
 	}
 }
 
+UObject* UAnimGraphNode_Base::GetObjectBeingDebugged() const
+{
+	if (UBlueprint* Blueprint = FBlueprintEditorUtils::FindBlueprintForNode(this))
+	{
+		return Blueprint->GetObjectBeingDebugged();
+	}
+
+	return nullptr;
+}
+
 bool UAnimGraphNode_Base::IsPotentiallyBoundFunction(const FMemberReference& FunctionReference)
 {
 	return FunctionReference.GetMemberGuid().IsValid() || FunctionReference.GetMemberName() != NAME_None;
diff --git a/Engine/Source/Editor/AnimGraph/Private/EditModes/AnimDynamicsEditMode.cpp b/Engine/Source/Editor/AnimGraph/Private/EditModes/AnimDynamicsEditMode.cpp
--- a/Engine/Source/Editor/AnimGraph/Private/EditModes/AnimDynamicsEditMode.cpp
+++ b/Engine/Source/Editor/AnimGraph/Private/EditModes/AnimDynamicsEditMode.cpp
@@ -117,21 +117,6 @@ FAnimDynamicsEditMode::FAnimDynamicsEditMode()
 	, bIsInteractingWithWidget(false)
 {}
 
-void FAnimDynamicsEditMode::EnterMode(class UAnimGraphNode_Base* InEditorNode, struct FAnimNode_Base* InRuntimeNode)
-{
-	FAnimNodeEditMode::EnterMode(InEditorNode, InRuntimeNode);
-
-	// Cache pointer to PreviewSkelMesh in the Editor node as there's no good way to access it directly from the Editor Node and its required for some parts of the UI.
-	{
-		USkeletalMeshComponent* const PreviewSkelMeshComp = GetAnimPreviewScene().GetPreviewMeshComponent();
-
-		if (UAnimGraphNode_AnimDynamics* const EditorAnimDynamicsNode = CastChecked<UAnimGraphNode_AnimDynamics>(InEditorNode))
-		{
-			EditorAnimDynamicsNode->LastPreviewComponent = PreviewSkelMeshComp;
-		}
-	}
-}
-
 void FAnimDynamicsEditMode::ExitMode()
 {
 	SelectedViewportObjects.Empty();
@@ -165,11 +150,6 @@ void FAnimDynamicsEditMode::Render(const FSceneView* View, FViewport* Viewport,
 
 		if (EditorAnimDynamicsNode)
 		{
-			if (EditorAnimDynamicsNode->LastPreviewComponent != PreviewSkelMeshComp)
-			{
-				EditorAnimDynamicsNode->LastPreviewComponent = PreviewSkelMeshComp;
-			}
-
 			const FAnimNode_AnimDynamics* const ActivePreviewNode = EditorAnimDynamicsNode->GetPreviewDynamicsNode();
 
 			if (!ActivePreviewNode)
@@ -830,8 +810,6 @@ void FAnimDynamicsEditMode::Tick(FEditorViewportClient* ViewportClient, float De
 					ActivePreviewNode->bDoPhysicsUpdateInEditor = true;
 				}
 			}
-
-			EditorAnimDynamicsNode->LastPreviewComponent = PreviewSkelMeshComp;
 		}
 	}
 }
diff --git a/Engine/Source/Editor/AnimGraph/Private/EditModes/AnimDynamicsEditMode.h b/Engine/Source/Editor/AnimGraph/Private/EditModes/AnimDynamicsEditMode.h
--- a/Engine/Source/Editor/AnimGraph/Private/EditModes/AnimDynamicsEditMode.h
+++ b/Engine/Source/Editor/AnimGraph/Private/EditModes/AnimDynamicsEditMode.h
@@ -30,7 +30,6 @@ public:
 	/** IAnimNodeEditMode interface */
 	virtual void Render(const FSceneView* View, FViewport* Viewport, FPrimitiveDrawInterface* PDI) override;
 	virtual bool HandleClick(FEditorViewportClient* InViewportClient, HHitProxy* HitProxy, const FViewportClick& Click) override;
-	virtual void EnterMode(class UAnimGraphNode_Base* InEditorNode, struct FAnimNode_Base* InRuntimeNode) override;
 	virtual void ExitMode() override;
 	virtual ECoordSystem GetWidgetCoordinateSystem() const override;
 	virtual FVector GetWidgetLocation() const override;
diff --git a/Engine/Source/Editor/AnimGraph/Public/AnimGraphNode_AnimDynamics.h b/Engine/Source/Editor/AnimGraph/Public/AnimGraphNode_AnimDynamics.h
--- a/Engine/Source/Editor/AnimGraph/Public/AnimGraphNode_AnimDynamics.h
+++ b/Engine/Source/Editor/AnimGraph/Public/AnimGraphNode_AnimDynamics.h
@@ -63,9 +63,6 @@ public:
 	UPROPERTY(EditAnywhere, Category = Preview)
 	bool bShowCollisionSpheres;
 
-	UPROPERTY(Transient)
-	mutable TObjectPtr<USkeletalMeshComponent> LastPreviewComponent;
-
 	virtual void PostLoad() override;
 
 	static FReply ResetButtonClicked(IDetailLayoutBuilder* DetailLayoutBuilder);
@@ -102,6 +99,10 @@ protected:
 	virtual const FAnimNode_SkeletalControlBase* GetNode() const override { return &Node; }
 	// End of UAnimGraphNode_SkeletalControlBase protected interface
 
+	USkeleton* GetSkeleton() const;
+
 private:
+	TSharedPtr<class FAnimGraphNode_AnimDynamics_DetailCustomization> DetailCustomization;
+
 	FNodeTitleTextTable CachedNodeTitles;
 };
diff --git a/Engine/Source/Editor/AnimGraph/Public/AnimGraphNode_Base.h b/Engine/Source/Editor/AnimGraph/Public/AnimGraphNode_Base.h
--- a/Engine/Source/Editor/AnimGraph/Public/AnimGraphNode_Base.h
+++ b/Engine/Source/Editor/AnimGraph/Public/AnimGraphNode_Base.h
@@ -500,8 +500,12 @@ public:
 
 	// Get the currently-debugged runtime anim node (in the anim BP debugger that this node is currently being edited in)
 	// @return nullptr if the node cannot be found
-	FAnimNode_Base* GetDebuggedAnimNode() const;
-	
+	FAnimNode_Base* GetDebuggedAnimNode() const { return GetDebuggedAnimNode<FAnimNode_Base>(); }
+
+	// Get the currently-debugged runtime anim node of a specified type (in the anim BP debugger that this node is currently being edited in)
+	// @return nullptr if the node cannot be found
+	template< typename TNodeType > TNodeType* GetDebuggedAnimNode() const;
+
 	// Refreshes the debugged component post-edit.
 	// This is required to see changes as the component may be either an editor-only component that is not ticking,
 	// or in a paused PIE world
@@ -592,6 +596,9 @@ protected:
 	// Helper function used to refresh the type of a binding
 	void RecalculateBindingType(FAnimGraphNodePropertyBinding& InBinding);
 	
+	/** @return the current object being debugged from the blueprint for this node. Can be nullptr. */
+	UObject* GetObjectBeingDebugged() const;
+
 protected:
 	// Old shown pins. Needs to be a member variable to track pin visibility changes between Pre and PostEditChange 
 	TArray<FName> OldShownPins;
@@ -614,3 +621,16 @@ void UAnimGraphNode_Base::HandleAnimReferenceReplacement(TObjectPtr<AssetType>&
 {
 	HandleAnimReferenceReplacement(static_cast<AssetType*&>(OriginalAsset), AnimAssetReplacementMap);
 }
+
+template<class TNodeType> TNodeType* UAnimGraphNode_Base::GetDebuggedAnimNode() const
+{
+	if (UObject* ActiveObject = GetObjectBeingDebugged())
+	{
+		if (UAnimBlueprintGeneratedClass* Class = Cast<UAnimBlueprintGeneratedClass>((UObject*)ActiveObject->GetClass()))
+		{
+			return Class->GetPropertyInstance<TNodeType>(ActiveObject, this);
+		}
+	}
+
+	return nullptr;
+}
diff --git a/Engine/Source/Runtime/AnimGraphRuntime/Private/BoneControllers/AnimNode_AnimDynamics.cpp b/Engine/Source/Runtime/AnimGraphRuntime/Private/BoneControllers/AnimNode_AnimDynamics.cpp
--- a/Engine/Source/Runtime/AnimGraphRuntime/Private/BoneControllers/AnimNode_AnimDynamics.cpp
+++ b/Engine/Source/Runtime/AnimGraphRuntime/Private/BoneControllers/AnimNode_AnimDynamics.cpp
@@ -460,17 +460,7 @@ void FAnimNode_AnimDynamics::EvaluateSkeletalControl_AnyThread(FComponentSpacePo
 	}
 }
 
-void FAnimNode_AnimDynamics::UpdateChainPhysicsBodyDefinitions(const USkeletalMeshComponent* const PreviewSkelMeshComp)
-{
-	if (PreviewSkelMeshComp)
-	{
-		TArray<FName> ChainBoneNames;
-		FindChainBones(BoundBone.BoneName, ChainEnd.BoneName, [PreviewSkelMeshComp](const FName BoneName) { return PreviewSkelMeshComp->GetParentBone(BoneName); }, ChainBoneNames);
-		UpdateChainPhysicsBodyDefinitions(ChainBoneNames);
-	}
-}
-
-void FAnimNode_AnimDynamics::UpdateChainPhysicsBodyDefinitions(const FReferenceSkeleton& ReferenceSkeleton)
+void FAnimNode_AnimDynamics::FindChainBoneNames(const FReferenceSkeleton& ReferenceSkeleton, TArray<FName>& ChainBoneNames)
 {
 	auto GetParentBoneNameFn = [&ReferenceSkeleton](const FName BoneName)
 	{
@@ -489,29 +479,44 @@ void FAnimNode_AnimDynamics::UpdateChainPhysicsBodyDefinitions(const FReferenceS
 		return FName(NAME_None);
 	};
 
-	TArray<FName> ChainBoneNames;
 	FindChainBones(BoundBone.BoneName, ChainEnd.BoneName, GetParentBoneNameFn, ChainBoneNames);
-	UpdateChainPhysicsBodyDefinitions(ChainBoneNames);
 }
 
-void FAnimNode_AnimDynamics::UpdateChainPhysicsBodyDefinitions(const TArray<FName>& ChainBoneNames)
+void FAnimNode_AnimDynamics::ValidateChainPhysicsBodyDefinitions(const FReferenceSkeleton& ReferenceSkeleton)
+{
+	TArray<FName> ChainBoneNames;
+	FindChainBoneNames(ReferenceSkeleton, ChainBoneNames);
+
+	// If another array of physics bodies has been pasted over this one we may have too many entries (in which case SetNum will truncate the array) or we may have too few (in which case SetNum will padd with default constucted body defs).
+	PhysicsBodyDefinitions.SetNum(ChainBoneNames.Num());
+
+	for (uint32 BodyIndex = 0, BodyIndexMax = FMath::Min(ChainBoneNames.Num(), PhysicsBodyDefinitions.Num()); BodyIndex < BodyIndexMax; ++BodyIndex)
+	{
+		PhysicsBodyDefinitions[BodyIndex].BoundBone.BoneName = ChainBoneNames[BodyIndex];
+	}
+}
+
+void FAnimNode_AnimDynamics::UpdateChainPhysicsBodyDefinitions(const FReferenceSkeleton& ReferenceSkeleton)
 {
+	TArray<FName> ChainBoneNames;
+	FindChainBoneNames(ReferenceSkeleton, ChainBoneNames);
 	check(ChainBoneNames.Num() > 0);
 
+	// If there was only one physics body then copy its values to all new chain bodies (emulating legacy behaviour), otherwise use values from default construction.
+	FAnimPhysBodyDefinition PrototypePhysBodyDef;
+
+	if (PhysicsBodyDefinitions.Num() == 1)
+	{
+		PrototypePhysBodyDef = PhysicsBodyDefinitions[0];
+	}
+
+
 	// Remove any bodies for bones that are not in the chain.
 	PhysicsBodyDefinitions.RemoveAll([&ChainBoneNames](const FAnimPhysBodyDefinition& Value) { return ChainBoneNames.Find(Value.BoundBone.BoneName) == INDEX_NONE;});
 	PhysicsBodyDefinitions.Reserve(ChainBoneNames.Num());
 
 	// Create a new Physics Body Def for any new bones in the chain and add them before or after the existing bodies as appropriate to maintain the order of the chain bones.
 	{
-		// If there was only one physics body then copy its values to all new chain bodies (emulating legacy behaviour), otherwise use values from default construction.
-		FAnimPhysBodyDefinition PrototypePhysBodyDef;
-		
-		if (PhysicsBodyDefinitions.Num() == 1)
-		{
-			PrototypePhysBodyDef = PhysicsBodyDefinitions[0];
-		}
-
 		uint32 PhysicsBodyDefIndex = 0;
 
 		for (FName BoneName : ChainBoneNames)
diff --git a/Engine/Source/Runtime/AnimGraphRuntime/Public/BoneControllers/AnimNode_AnimDynamics.h b/Engine/Source/Runtime/AnimGraphRuntime/Public/BoneControllers/AnimNode_AnimDynamics.h
--- a/Engine/Source/Runtime/AnimGraphRuntime/Public/BoneControllers/AnimNode_AnimDynamics.h
+++ b/Engine/Source/Runtime/AnimGraphRuntime/Public/BoneControllers/AnimNode_AnimDynamics.h
@@ -292,13 +292,13 @@ struct ANIMGRAPHRUNTIME_API FAnimNode_AnimDynamics : public FAnimNode_SkeletalCo
 	/**
 	* Overridden linear damping value. The default is 0.7. Values below 0.7 won't have an effect.
 	*/
-	UPROPERTY(EditAnywhere, AdvancedDisplay, BlueprintReadWrite, Category = Setup, meta = (PinHiddenByDefault, DisplayAfter = "NumSolverIterationsPostUpdate", EditCondition = bOverrideLinearDamping))
+	UPROPERTY(EditAnywhere, AdvancedDisplay, BlueprintReadWrite, Category = PhysicsParameters, meta = (PinHiddenByDefault, DisplayAfter = "NumSolverIterationsPostUpdate", EditCondition = bOverrideLinearDamping))
 	float LinearDampingOverride;
 
 	/**
 	 * Overridden angular damping value. The default is 0.7. Values below 0.7 won't have an effect.
 	 */
-	UPROPERTY(EditAnywhere, AdvancedDisplay, BlueprintReadWrite, Category = Setup, meta = (PinHiddenByDefault, DisplayAfter = "LinearDampingOverride", EditCondition = bOverrideAngularDamping))
+	UPROPERTY(EditAnywhere, AdvancedDisplay, BlueprintReadWrite, Category = PhysicsParameters, meta = (PinHiddenByDefault, DisplayAfter = "LinearDampingOverride", EditCondition = bOverrideAngularDamping))
 	float AngularDampingOverride;
 
 	// Previous component & actor transforms, used to account for teleports
@@ -317,22 +317,22 @@ struct ANIMGRAPHRUNTIME_API FAnimNode_AnimDynamics : public FAnimNode_SkeletalCo
 	UPROPERTY(EditAnywhere, Category = Setup, meta = (EditCondition = bChain, DisplayAfter = "BoundBone"))
 	FBoneReference ChainEnd;
 
-	UPROPERTY(EditAnywhere, EditFixedSize, Category = Setup, meta = (DisplayName = "Physics Bodies", EditFixedOrder, DisplayAfter = "ChainEnd"))
+	UPROPERTY(EditAnywhere, EditFixedSize, Category = PhysicsParameters, meta = (DisplayName = "Body Definitions", EditFixedOrder, DisplayAfter = "ChainEnd"))
 	TArray< FAnimPhysBodyDefinition > PhysicsBodyDefinitions;
 
 	/** Scale for gravity, higher values increase forces due to gravity */
-	UPROPERTY(EditAnywhere, BlueprintReadWrite, Category = Setup, meta = (PinHiddenByDefault, DisplayAfter = "bGravityOverrideInSimSpace", EditCondition = "!bUseGravityOverride"))
+	UPROPERTY(EditAnywhere, BlueprintReadWrite, Category = PhysicsParameters, meta = (PinHiddenByDefault, DisplayAfter = "bGravityOverrideInSimSpace", EditCondition = "!bUseGravityOverride"))
 	float GravityScale;
 
 	/** Gravity Override Value */
-	UPROPERTY(EditAnywhere, BlueprintReadWrite, Category = Setup, meta = (PinHiddenByDefault, DisplayAfter = "bUseGravityOverride", EditCondition = "bUseGravityOverride"))
+	UPROPERTY(EditAnywhere, BlueprintReadWrite, Category = PhysicsParameters, meta = (PinHiddenByDefault, DisplayAfter = "bUseGravityOverride", EditCondition = "bUseGravityOverride"))
 	FVector GravityOverride;
 
 	/** 
 	 * Spring constant to use when calculating linear springs, higher values mean a stronger spring.
 	 * You need to enable the Linear Spring checkbox for this to have an effect.
 	 */
-	UPROPERTY(EditAnywhere, BlueprintReadWrite, Category = Setup, meta = (EditCondition = bLinearSpring, PinHiddenByDefault, DisplayAfter = "AngularSpringConstant"))
+	UPROPERTY(EditAnywhere, BlueprintReadWrite, Category = PhysicsParameters, meta = (EditCondition = bLinearSpring, PinHiddenByDefault, DisplayAfter = "AngularSpringConstant"))
 	float LinearSpringConstant;
 
 	/** 
@@ -340,7 +340,7 @@ struct ANIMGRAPHRUNTIME_API FAnimNode_AnimDynamics : public FAnimNode_SkeletalCo
 	 * You need to enable the Angular Spring checkbox for this to have an effect.
 	 * Note: Make sure to also set the Angular Target Axis and Angular Target in the Constraint Setup for this to have an effect.
 	 */
-	UPROPERTY(EditAnywhere, BlueprintReadWrite, Category = Setup, meta = (DisplayAfter = "PhysicsBodyDefinitions", EditCondition = bAngularSpring, PinHiddenByDefault))
+	UPROPERTY(EditAnywhere, BlueprintReadWrite, Category = PhysicsParameters, meta = (DisplayAfter = "PhysicsBodyDefinitions", EditCondition = bAngularSpring, PinHiddenByDefault))
 	float AngularSpringConstant;
 
 	/** Scale to apply to calculated wind velocities in the solver */
@@ -364,15 +364,15 @@ struct ANIMGRAPHRUNTIME_API FAnimNode_AnimDynamics : public FAnimNode_SkeletalCo
 	*  in check. When using single-body systems sometimes angular forces will look like they are "catching-up" with
 	*  the mesh, if that's the case override this and push it towards 1.0f until it settles correctly
 	*/
-	UPROPERTY(EditAnywhere, AdvancedDisplay, BlueprintReadWrite, Category = Setup, meta = (PinHiddenByDefault, DisplayAfter = "AngularDampingOverride", EditCondition = bOverrideAngularBias))
+	UPROPERTY(EditAnywhere, AdvancedDisplay, BlueprintReadWrite, Category = PhysicsParameters, meta = (PinHiddenByDefault, DisplayAfter = "AngularDampingOverride", EditCondition = bOverrideAngularBias))
 	float AngularBiasOverride;
 
 	/** Number of update passes on the linear and angular limits before we solve the position of the bodies recommended to be four times the value of NumSolverIterationsPostUpdate */
-	UPROPERTY(EditAnywhere, AdvancedDisplay, Category = Setup)
+	UPROPERTY(EditAnywhere, AdvancedDisplay, Category = PhysicsParameters)
 	int32 NumSolverIterationsPreUpdate;
 
 	/** Number of update passes on the linear and angular limits after we solve the position of the bodies, recommended to be around a quarter of NumSolverIterationsPreUpdate */
-	UPROPERTY(EditAnywhere, AdvancedDisplay, Category = Setup, meta = (PinHiddenByDefault, DisplayAfter = "NumSolverIterationsPreUpdate"))
+	UPROPERTY(EditAnywhere, AdvancedDisplay, Category = PhysicsParameters, meta = (PinHiddenByDefault, DisplayAfter = "NumSolverIterationsPreUpdate"))
 	int32 NumSolverIterationsPostUpdate;
 
 	/** List of available spherical limits for this node */
@@ -407,15 +407,15 @@ struct ANIMGRAPHRUNTIME_API FAnimNode_AnimDynamics : public FAnimNode_SkeletalCo
 	uint8 bUsePlanarLimit:1;
 
 	/** If true we will perform physics update, otherwise skip - allows visualization of the initial state of the bodies */
-	UPROPERTY(EditAnywhere, AdvancedDisplay, Category = Setup, meta = (DisplayAfter = "bDoEval"))
+	UPROPERTY(EditAnywhere, AdvancedDisplay, Category = PhysicsParameters, meta = (DisplayAfter = "bDoEval"))
 	uint8 bDoUpdate : 1;
 
 	/** If true we will perform bone transform evaluation, otherwise skip - allows visualization of the initial anim state compared to the physics sim */
-	UPROPERTY(EditAnywhere, AdvancedDisplay, Category = Setup, meta = (DisplayAfter = "AngularBiasOverride"))
+	UPROPERTY(EditAnywhere, AdvancedDisplay, Category = PhysicsParameters, meta = (DisplayAfter = "AngularBiasOverride"))
 	uint8 bDoEval : 1;
 
 	/** If true, the override value will be used for linear damping */
-	UPROPERTY(EditAnywhere, AdvancedDisplay, Category = Setup, meta=(InlineEditConditionToggle, DisplayAfter="AngularSpringConstraint"))
+	UPROPERTY(EditAnywhere, AdvancedDisplay, Category = PhysicsParameters, meta=(InlineEditConditionToggle, DisplayAfter="AngularSpringConstraint"))
 	uint8 bOverrideLinearDamping:1;
 
 	/** If true, the override value will be used for the angular bias for bodies in this node. 
@@ -423,11 +423,11 @@ struct ANIMGRAPHRUNTIME_API FAnimNode_AnimDynamics : public FAnimNode_SkeletalCo
 	 *  in check. When using single-body systems sometimes angular forces will look like they are "catching-up" with
 	 *  the mesh, if that's the case override this and push it towards 1.0f until it settles correctly
 	 */
-	UPROPERTY(EditAnywhere, AdvancedDisplay, Category = Setup, meta=(InlineEditConditionToggle, DisplayAfter="AngularDampingOverride"))
+	UPROPERTY(EditAnywhere, AdvancedDisplay, Category = PhysicsParameters, meta=(InlineEditConditionToggle, DisplayAfter="AngularDampingOverride"))
 	uint8 bOverrideAngularBias:1;
 
 	/** If true, the override value will be used for angular damping */
-	UPROPERTY(EditAnywhere, AdvancedDisplay, Category = Setup, meta=(InlineEditConditionToggle, DisplayAfter="LinearDampingOverride"))
+	UPROPERTY(EditAnywhere, AdvancedDisplay, Category = PhysicsParameters, meta=(InlineEditConditionToggle, DisplayAfter="LinearDampingOverride"))
 	uint8 bOverrideAngularDamping:1;
 
 	/** Whether or not wind is enabled for the bodies in this simulation */
@@ -437,19 +437,19 @@ struct ANIMGRAPHRUNTIME_API FAnimNode_AnimDynamics : public FAnimNode_SkeletalCo
 	uint8 bWindWasEnabled:1;
 
 	/** Use gravity override value vs gravity scale */
-	UPROPERTY(EditAnywhere, BlueprintReadWrite, Category = Setup, meta = (DisplayAfter = "LinearSpringConstant"))
+	UPROPERTY(EditAnywhere, BlueprintReadWrite, Category = PhysicsParameters, meta = (DisplayAfter = "LinearSpringConstant"))
 	uint8 bUseGravityOverride:1;
 
 	/** If true the gravity override value is defined in simulation space, by default it is in world space */
-	UPROPERTY(EditAnywhere, Category = Setup, meta=(DisplayAfter = "GravityOverride", DisplayName = "Gravity Override In Sim Space", EditCondition = "bUseGravityOverride"))
+	UPROPERTY(EditAnywhere, Category = PhysicsParameters, meta=(DisplayAfter = "GravityOverride", DisplayName = "Gravity Override In Sim Space", EditCondition = "bUseGravityOverride"))
 	uint8 bGravityOverrideInSimSpace : 1;
 
 	/** If true the body will attempt to spring back to its initial position */
-	UPROPERTY(EditAnywhere, Category = Setup, meta = (InlineEditConditionToggle))
+	UPROPERTY(EditAnywhere, Category = PhysicsParameters, meta = (InlineEditConditionToggle))
 	uint8 bLinearSpring:1;
 
 	/** If true the body will attempt to align itself with the specified angular target */
-	UPROPERTY(EditAnywhere, Category = Setup, meta = (InlineEditConditionToggle))
+	UPROPERTY(EditAnywhere, Category = PhysicsParameters, meta = (InlineEditConditionToggle))
 	uint8 bAngularSpring:1;
 
 	/** Set to true to use the solver to simulate a connected chain */
@@ -489,9 +489,9 @@ struct ANIMGRAPHRUNTIME_API FAnimNode_AnimDynamics : public FAnimNode_SkeletalCo
 	void InitPhysics(FComponentSpacePoseContext& Output);
 	void TermPhysics();
 
-	void UpdateChainPhysicsBodyDefinitions(const USkeletalMeshComponent* const PreviewSkelMeshComp);
 	void UpdateChainPhysicsBodyDefinitions(const FReferenceSkeleton& ReferenceSkeleton);
-	void UpdateChainPhysicsBodyDefinitions(const TArray<FName>& ChainBoneNames);
+	void ValidateChainPhysicsBodyDefinitions(const FReferenceSkeleton& ReferenceSkeleton);
+	void FindChainBoneNames(const FReferenceSkeleton& ReferenceSkeleton, TArray<FName>& ChainBoneNames);
 	void UpdateLimits(FComponentSpacePoseContext& Output);
 
 	int32 GetNumBodies() const;
