Added missing exposed value handler update in FAnimNode_SequenceEvaluatorBase::Initialize_AnyThread
#rnx

diff --git a/Engine/Source/Runtime/AnimGraphRuntime/Private/AnimNodes/AnimNode_SequenceEvaluator.cpp b/Engine/Source/Runtime/AnimGraphRuntime/Private/AnimNodes/AnimNode_SequenceEvaluator.cpp
--- a/Engine/Source/Runtime/AnimGraphRuntime/Private/AnimNodes/AnimNode_SequenceEvaluator.cpp
+++ b/Engine/Source/Runtime/AnimGraphRuntime/Private/AnimNodes/AnimNode_SequenceEvaluator.cpp
@@ -19,6 +19,7 @@ void FAnimNode_SequenceEvaluatorBase::Initialize_AnyThread(const FAnimationIniti
 {
 	DECLARE_SCOPE_HIERARCHICAL_COUNTER_ANIMNODE(Initialize_AnyThread)
 	FAnimNode_AssetPlayerBase::Initialize_AnyThread(Context);
+	GetEvaluateGraphExposedInputs().Execute(Context);
 	bReinitialized = true;
 }
 
