Fix for non-unity compile error in AnimationStateMachineLibrary
#rnx

diff --git a/Engine/Source/Runtime/AnimGraphRuntime/Public/AnimationStateMachineLibrary.h b/Engine/Source/Runtime/AnimGraphRuntime/Public/AnimationStateMachineLibrary.h
--- a/Engine/Source/Runtime/AnimGraphRuntime/Public/AnimationStateMachineLibrary.h
+++ b/Engine/Source/Runtime/AnimGraphRuntime/Public/AnimationStateMachineLibrary.h
@@ -7,6 +7,7 @@
 #include "Animation/AnimExecutionContext.h"
 #include "Kismet/BlueprintFunctionLibrary.h"
 #include "Animation/AnimNodeReference.h"
+#include "Animation/AnimStateMachineTypes.h"
 #include "AnimationStateMachineLibrary.generated.h"
 
 struct FAnimNode_StateMachine;
