[UE][FIX]] AnimDynamics - Ensure all parts of BodyDef are copied on chain update.

#rmx

diff --git a/Engine/Source/Runtime/AnimGraphRuntime/Private/BoneControllers/AnimNode_AnimDynamics.cpp b/Engine/Source/Runtime/AnimGraphRuntime/Private/BoneControllers/AnimNode_AnimDynamics.cpp
--- a/Engine/Source/Runtime/AnimGraphRuntime/Private/BoneControllers/AnimNode_AnimDynamics.cpp
+++ b/Engine/Source/Runtime/AnimGraphRuntime/Private/BoneControllers/AnimNode_AnimDynamics.cpp
@@ -524,12 +524,8 @@ void FAnimNode_AnimDynamics::UpdateChainPhysicsBodyDefinitions(const FReferenceS
 			if (!PhysicsBodyDefinitions.FindByPredicate([BoneName](const FAnimPhysBodyDefinition& Value) { return Value.BoundBone.BoneName == BoneName; }))
 			{
 				// Add the new bone to the chain.
-				FAnimPhysBodyDefinition ChainPhysicsBody;
-				ChainPhysicsBody.BoundBone.BoneName = BoneName;
-				ChainPhysicsBody.BoxExtents = PrototypePhysBodyDef.BoxExtents;
-				ChainPhysicsBody.LocalJointOffset = PrototypePhysBodyDef.LocalJointOffset;
-				ChainPhysicsBody.SphereCollisionRadius = PrototypePhysBodyDef.SphereCollisionRadius;
-				PhysicsBodyDefinitions.Insert(ChainPhysicsBody, PhysicsBodyDefIndex);
+				PhysicsBodyDefinitions.Insert(PrototypePhysBodyDef, PhysicsBodyDefIndex);
+				PhysicsBodyDefinitions[PhysicsBodyDefIndex].BoundBone.BoneName = BoneName;
 			}
 
 			++PhysicsBodyDefIndex;
