Fixed a bug where mirroring did not work correctly with some skeletons
Cherry pick with additional edits from CL18522092, CL18536881, CL18811714

diff --git a/Engine/Source/Runtime/AnimGraphRuntime/Private/AnimNodes/AnimNode_Mirror.cpp b/Engine/Source/Runtime/AnimGraphRuntime/Private/AnimNodes/AnimNode_Mirror.cpp
--- a/Engine/Source/Runtime/AnimGraphRuntime/Private/AnimNodes/AnimNode_Mirror.cpp
+++ b/Engine/Source/Runtime/AnimGraphRuntime/Private/AnimNodes/AnimNode_Mirror.cpp
@@ -31,10 +31,6 @@ bool FAnimNode_MirrorBase::GetMirror() const
 	return false;
 }
 
-UBlendProfile* FAnimNode_MirrorBase::GetBlendProfile() const
-{
-	return nullptr;
-}
 
 float FAnimNode_MirrorBase::GetBlendTimeOnMirrorStateChange() const
 {
@@ -65,12 +61,6 @@ bool FAnimNode_MirrorBase::SetMirror(bool bInMirror)
 {
 	return false;
 }
-
-bool FAnimNode_MirrorBase::SetBlendProfile(UBlendProfile* BlendProfile)
-{
-	return false;
-}
-
 bool FAnimNode_MirrorBase::SetBlendTimeOnMirrorStateChange(float InBlendTime)
 {
 	return false;
@@ -151,7 +141,7 @@ void FAnimNode_MirrorBase::Update_AnyThread(const FAnimationUpdateContext& Conte
 			UE::Anim::IInertializationRequester* InertializationRequester = Context.GetMessage<UE::Anim::IInertializationRequester>();
 			if (InertializationRequester)
 			{
-				InertializationRequester->RequestInertialization(GetBlendTimeOnMirrorStateChange(), GetBlendProfile());
+				InertializationRequester->RequestInertialization(GetBlendTimeOnMirrorStateChange());
 				InertializationRequester->AddDebugRecord(*Context.AnimInstanceProxy, Context.GetCurrentNodeId());
 			}
 			else
@@ -248,11 +238,6 @@ bool FAnimNode_Mirror::GetMirror() const
 	return GET_ANIM_NODE_DATA(bool, bMirror);
 }
 
-UBlendProfile* FAnimNode_Mirror::GetBlendProfile() const
-{
-	return GET_ANIM_NODE_DATA(TObjectPtr<UBlendProfile>, BlendProfile);
-}
-
 float FAnimNode_Mirror::GetBlendTimeOnMirrorStateChange() const
 {
 	return GET_ANIM_NODE_DATA(float, BlendTime);
@@ -291,23 +276,6 @@ bool FAnimNode_Mirror::SetMirror(bool bInMirror)
 	}
 	return false;
 }
-
-bool FAnimNode_Mirror::SetBlendProfile(UBlendProfile* InBlendProfile)
-{
-#if WITH_EDITORONLY_DATA
-	BlendProfile = InBlendProfile;
-	GET_MUTABLE_ANIM_NODE_DATA(TObjectPtr<UBlendProfile>, BlendProfile) = InBlendProfile;
-#endif
-
-	if (TObjectPtr<UBlendProfile>* BlendProfilePtr = GET_INSTANCE_ANIM_NODE_DATA_PTR(TObjectPtr<UBlendProfile>, BlendProfile))
-	{
-		*BlendProfilePtr = InBlendProfile;
-		return true;
-	}
-
-	return false;
-}
-
 bool FAnimNode_Mirror::SetBlendTimeOnMirrorStateChange(float InBlendTime)
 {
 #if WITH_EDITORONLY_DATA
@@ -403,11 +371,6 @@ bool FAnimNode_Mirror_Standalone::GetMirror() const
 	return bMirror;
 }
 
-UBlendProfile* FAnimNode_Mirror_Standalone::GetBlendProfile() const
-{
-	return BlendProfile.Get();
-}
-
 float FAnimNode_Mirror_Standalone::GetBlendTimeOnMirrorStateChange() const
 {
 	return BlendTime;
@@ -438,13 +401,6 @@ bool FAnimNode_Mirror_Standalone::SetMirror(bool bInMirror)
 	bMirror = bInMirror;
 	return true;
 }
-
-bool FAnimNode_Mirror_Standalone::SetBlendProfile(UBlendProfile* InBlendProfile)
-{
-	BlendProfile = InBlendProfile;
-	return true;
-}
-
 bool FAnimNode_Mirror_Standalone::SetBlendTimeOnMirrorStateChange(float InBlendTime)
 {
 	BlendTime = InBlendTime;
diff --git a/Engine/Source/Runtime/AnimGraphRuntime/Public/AnimNodes/AnimNode_Mirror.h b/Engine/Source/Runtime/AnimGraphRuntime/Public/AnimNodes/AnimNode_Mirror.h
--- a/Engine/Source/Runtime/AnimGraphRuntime/Public/AnimNodes/AnimNode_Mirror.h
+++ b/Engine/Source/Runtime/AnimGraphRuntime/Public/AnimNodes/AnimNode_Mirror.h
@@ -29,10 +29,6 @@ public:
 
 	// Get Mirror State
 	virtual bool GetMirror() const;
-
-	// Get the BlendProfile to use when transitioning between mirrored and unmirrored states
-	virtual UBlendProfile* GetBlendProfile() const;
-
 	// How long to blend using inertialization when switching  mirrored state
 	virtual float GetBlendTimeOnMirrorStateChange() const;
 
@@ -51,9 +47,6 @@ public:
 	// Set Mirror State
 	virtual bool SetMirror(bool bInMirror);
 
-	// Set the BlendProfile to use when transitioning between mirrored and unmirrored states
-	virtual bool SetBlendProfile(UBlendProfile* BlendProfile);
-
 	// Set how long to blend using inertialization when switching  mirrored state
 	// @return true if the value was set (it is dynamic), or false if it could not (it is not dynamic or pin exposed)
 	virtual bool SetBlendTimeOnMirrorStateChange(float InBlendTime);
@@ -110,7 +103,6 @@ public:
 	virtual bool SetMirrorDataTable(UMirrorDataTable* MirrorTable) override;
 
 	virtual bool GetMirror() const override;
-	virtual UBlendProfile* GetBlendProfile() const;
 	virtual float GetBlendTimeOnMirrorStateChange() const override;
 	virtual bool GetBoneMirroring() const override;
 	virtual bool GetCurveMirroring() const override;
@@ -118,7 +110,6 @@ public:
 	virtual bool GetResetChildOnMirrorStateChange() const override;
 
 	virtual bool SetMirror(bool bInMirror) override;
-	virtual bool SetBlendProfile(UBlendProfile* BlendProfile);
 	virtual bool SetBlendTimeOnMirrorStateChange(float InBlendTime) override;
 	virtual bool SetBoneMirroring(bool bInBoneMirroring) override;
 	virtual bool SetCurveMirroring(bool bInCurveMirroring) override;
@@ -137,11 +128,6 @@ protected:
 	// Inertialization blend time to use when transitioning between mirrored and unmirrored states
 	UPROPERTY(EditAnywhere, Category = MirrorTransition, meta = (PinHiddenByDefault, FoldProperty))
 	float BlendTime = 0.0f;
-
-	// Blend profile to use when transitioning between mirrored and unmirrored states
-	UPROPERTY(EditAnywhere, Category = MirrorTransition, meta = (PinHiddenByDefault, UseAsBlendProfile = true, FoldProperty))
-	TObjectPtr<UBlendProfile> BlendProfile = nullptr;
-
 	// Whether to reset (reinitialize) the child (source) pose when the mirror state changes
 	UPROPERTY(EditAnywhere, Category = MirrorTransition, meta = (FoldProperty))
 	bool bResetChild = false;
@@ -172,7 +158,6 @@ public:
 	virtual bool SetMirrorDataTable(UMirrorDataTable* MirrorTable) override;
 
 	virtual bool GetMirror() const override;
-	virtual UBlendProfile* GetBlendProfile() const;
 	virtual float GetBlendTimeOnMirrorStateChange() const override;
 	virtual bool GetBoneMirroring() const override;
 	virtual bool GetCurveMirroring() const override;
@@ -180,7 +165,6 @@ public:
 	virtual bool GetResetChildOnMirrorStateChange() const override;
 
 	virtual bool SetMirror(bool bInMirror) override;
-	virtual bool SetBlendProfile(UBlendProfile* BlendProfile);
 	virtual bool SetBlendTimeOnMirrorStateChange(float InBlendTime) override;
 	virtual bool SetBoneMirroring(bool bInBoneMirroring) override;
 	virtual bool SetCurveMirroring(bool bInCurveMirroring) override;
@@ -199,10 +183,6 @@ protected:
 	UPROPERTY(EditAnywhere, Category = MirrorTransition, meta = (PinHiddenByDefault, FoldProperty))
 	float BlendTime = 0.0f;
 
-	// Blend profile to use when transitioning between mirrored and unmirrored states
-	UPROPERTY(EditAnywhere, Category = MirrorTransition, meta = (PinHiddenByDefault, UseAsBlendProfile = true, FoldProperty))
-	TObjectPtr<UBlendProfile> BlendProfile = nullptr;
-
 	// Whether to reset (reinitialize) the child (source) pose when the mirror state changes
 	UPROPERTY(EditAnywhere, Category = MirrorTransition, meta = (FoldProperty))
 	bool bResetChild = false;
