Tweak tooltip for sequence player's SetStartPosition to better explain its functionality

diff --git a/Engine/Source/Runtime/AnimGraphRuntime/Public/SequencePlayerLibrary.h b/Engine/Source/Runtime/AnimGraphRuntime/Public/SequencePlayerLibrary.h
--- a/Engine/Source/Runtime/AnimGraphRuntime/Public/SequencePlayerLibrary.h
+++ b/Engine/Source/Runtime/AnimGraphRuntime/Public/SequencePlayerLibrary.h
@@ -44,7 +44,11 @@ public:
 	UFUNCTION(BlueprintCallable, Category = "Sequence Player", meta = (BlueprintThreadSafe))
 	static FSequencePlayerReference SetAccumulatedTime(const FSequencePlayerReference& SequencePlayer, float Time);
 
-	/** Set the start position of the sequence player */
+	/** 
+	 * Set the start position of the sequence player. 
+	 * If this is called from On Become Relevant or On Initial Update then it should be accompanied by a call to
+	 * SetAccumulatedTime to achieve the desired effect of resetting the play time of a sequence player.
+	 */
 	UFUNCTION(BlueprintCallable, Category = "Sequence Player", meta = (BlueprintThreadSafe))
 	static FSequencePlayerReference SetStartPosition(const FSequencePlayerReference& SequencePlayer, float StartPosition);
 
