Fix for incorrect behavior in Spline IK node

diff --git a/Engine/Source/Runtime/AnimGraphRuntime/Private/BoneControllers/AnimNode_SplineIK.cpp b/Engine/Source/Runtime/AnimGraphRuntime/Private/BoneControllers/AnimNode_SplineIK.cpp
--- a/Engine/Source/Runtime/AnimGraphRuntime/Private/BoneControllers/AnimNode_SplineIK.cpp
+++ b/Engine/Source/Runtime/AnimGraphRuntime/Private/BoneControllers/AnimNode_SplineIK.cpp
@@ -370,11 +370,12 @@ void FAnimNode_SplineIK::BuildBoneSpline(const FReferenceSkeleton& RefSkeleton)
 			float TotalPointCount = (float)(ClampedPointCount - 1);
 			for (int32 PointIndex = 0; PointIndex < ClampedPointCount; ++PointIndex)
 			{
-				const float CurveAlpha = (float)PointIndex / TotalPointCount;
+				const float TransformedCurveAlpha = (float)PointIndex / TotalPointCount;
+				const float CurveAlpha = (float)PointIndex;
 
-				BoneSpline.Position.Points.Emplace(CurveAlpha, TransformedSpline.Position.Eval(CurveAlpha), FVector::ZeroVector, FVector::ZeroVector, CIM_CurveAuto);
-				BoneSpline.Rotation.Points.Emplace(CurveAlpha, TransformedSpline.Rotation.Eval(CurveAlpha), FQuat::Identity, FQuat::Identity, CIM_Linear);
-				BoneSpline.Scale.Points.Emplace(CurveAlpha, TransformedSpline.Scale.Eval(CurveAlpha), FVector::ZeroVector, FVector::ZeroVector, CIM_CurveAuto);
+				BoneSpline.Position.Points.Emplace(CurveAlpha, TransformedSpline.Position.Eval(TransformedCurveAlpha), FVector::ZeroVector, FVector::ZeroVector, CIM_CurveAuto);
+				BoneSpline.Rotation.Points.Emplace(CurveAlpha, TransformedSpline.Rotation.Eval(TransformedCurveAlpha), FQuat::Identity, FQuat::Identity, CIM_Linear);
+				BoneSpline.Scale.Points.Emplace(CurveAlpha, TransformedSpline.Scale.Eval(TransformedCurveAlpha), FVector::ZeroVector, FVector::ZeroVector, CIM_CurveAuto);
 			}
 
 			// clear the transformed spline so we dont end up using it
