Added WithPureVirtual trait to FAnimNode_BlendSpacePlayerBase to fix compilation errors on NoPCH build configuration.

diff --git a/Engine/Source/Runtime/AnimGraphRuntime/Public/AnimNodes/AnimNode_BlendSpacePlayer.h b/Engine/Source/Runtime/AnimGraphRuntime/Public/AnimNodes/AnimNode_BlendSpacePlayer.h
--- a/Engine/Source/Runtime/AnimGraphRuntime/Public/AnimNodes/AnimNode_BlendSpacePlayer.h
+++ b/Engine/Source/Runtime/AnimGraphRuntime/Public/AnimNodes/AnimNode_BlendSpacePlayer.h
@@ -99,6 +99,15 @@ private:
 	const FBlendSampleData* GetHighestWeightedSample() const;
 };
 
+template<>
+struct TStructOpsTypeTraits<FAnimNode_BlendSpacePlayerBase> : public TStructOpsTypeTraitsBase2<FAnimNode_BlendSpacePlayerBase>
+{
+	enum
+	{
+		WithPureVirtual = true,
+	};
+};
+
 
 //@TODO: Comment
 USTRUCT(BlueprintInternalUseOnly)
