Added a flag on RigidBody AnimNode to allow it to use synchronous physics simulation. Renamed the CVar to reflect that it's the project default, not an override.

diff --git a/Engine/Source/Runtime/AnimGraphRuntime/Private/BoneControllers/AnimNode_RigidBody.cpp b/Engine/Source/Runtime/AnimGraphRuntime/Private/BoneControllers/AnimNode_RigidBody.cpp
--- a/Engine/Source/Runtime/AnimGraphRuntime/Private/BoneControllers/AnimNode_RigidBody.cpp
+++ b/Engine/Source/Runtime/AnimGraphRuntime/Private/BoneControllers/AnimNode_RigidBody.cpp
@@ -71,9 +71,14 @@ FAutoConsoleVariableRef CVarRigidBodyNodeSpaceExternalLinearVelocityY(TEXT("p.Ri
 FAutoConsoleVariableRef CVarRigidBodyNodeSpaceExternalLinearVelocityZ(TEXT("p.RigidBodyNode.Space.ExternalLinearVelocity.Z"), RBAN_SimSpaceOverride.ExternalLinearVelocity.Z, TEXT("RBAN SimSpaceSettings overrides"), ECVF_Default);
 #endif
 
-bool bRBAN_UseDeferredTask = false;
+bool bRBAN_DeferredSimulationDefault = false;
+FAutoConsoleVariableRef CVarRigidBodyNodeDeferredSimulationDefault(
+	TEXT("p.RigidBodyNode.DeferredSimulationDefault"),
+	bRBAN_DeferredSimulationDefault,
+	TEXT("Whether rigid body simulations are deferred one frame for assets that don't opt into a specific simulation timing"),
+	ECVF_Default);
+
 bool bRBAN_DebugDraw = false;
-FAutoConsoleVariableRef CVarRigidBodyNodeDeferredSimulation(TEXT("p.RigidBodyNode.UseDeferredTask"), bRBAN_UseDeferredTask, TEXT("Whether to defer the simulation results by one frame so that they can run in a task"), ECVF_Default);
 FAutoConsoleVariableRef CVarRigidBodyNodeDebugDraw(TEXT("p.RigidBodyNode.DebugDraw"), bRBAN_DebugDraw, TEXT("Whether to debug draw the rigid body simulation state. Requires p.Chaos.DebugDraw.Enabled 1 to function as well."), ECVF_Default);
 
 // Array of priorities that can be indexed into with CVars, since task priorities cannot be set from scalability .ini
@@ -696,9 +701,10 @@ void FAnimNode_RigidBody::EvaluateSkeletalControl_AnyThread(FComponentSpacePoseC
 			PreviousComponentLinearVelocity = FVector::ZeroVector;
 		}
 
-		// Run simulation at a minimum of 30 FPS to prevent system from exploding.
-		// DeltaTime can be higher due to URO, so take multiple iterations in that case.
-		const bool bUseDeferredSimulationTask = bRBAN_UseDeferredTask;
+		// Assets can override config for deferred simulation
+		const bool bUseDeferredSimulationTask =
+			(SimulationTiming == ESimulationTiming::Deferred) ||
+			((SimulationTiming == ESimulationTiming::Default) && bRBAN_DeferredSimulationDefault);
 		FVector SimSpaceGravity(0.f);
 
 		// Only need to tick physics if we didn't reset and we have some time to simulate
diff --git a/Engine/Source/Runtime/AnimGraphRuntime/Public/BoneControllers/AnimNode_RigidBody.h b/Engine/Source/Runtime/AnimGraphRuntime/Public/BoneControllers/AnimNode_RigidBody.h
--- a/Engine/Source/Runtime/AnimGraphRuntime/Public/BoneControllers/AnimNode_RigidBody.h
+++ b/Engine/Source/Runtime/AnimGraphRuntime/Public/BoneControllers/AnimNode_RigidBody.h
@@ -28,6 +28,17 @@ enum class ESimulationSpace : uint8
 	BaseBoneSpace,
 };
 
+/** Determines behaviour regarding deferral of simulation tasks. */
+UENUM()
+enum class ESimulationTiming : uint8
+{
+	/** Use the default project setting as defined by p.RigidBodyNode.DeferredSimulationDefault. */
+	Default,
+	/** Always run the simulation to completion during animation evaluation. */
+	Synchronous,
+	/** Always run the simulation in the background and retrieve the result on the next animation evaluation. */
+	Deferred
+};
 
 /**
  * Settings for the system which passes motion of the simulation's space into the simulation. This allows the simulation to pass a 
@@ -329,6 +340,12 @@ private:
 	// Destroy the simulation and free related structures
 	void DestroyPhysicsSimulation();
 
+public:
+
+	/* Whether the physics simulation runs synchronously with the node's evaluation or is run in the background until the next frame. */
+	UPROPERTY(EditAnywhere, Category=Settings, AdvancedDisplay)
+	ESimulationTiming SimulationTiming;
+
 private:
 
 	float WorldTimeSeconds;
