Fixing CIS warning for AnimNode_RigidBody and initializing it to default.

diff --git a/Engine/Source/Runtime/AnimGraphRuntime/Private/BoneControllers/AnimNode_RigidBody.cpp b/Engine/Source/Runtime/AnimGraphRuntime/Private/BoneControllers/AnimNode_RigidBody.cpp
--- a/Engine/Source/Runtime/AnimGraphRuntime/Private/BoneControllers/AnimNode_RigidBody.cpp
+++ b/Engine/Source/Runtime/AnimGraphRuntime/Private/BoneControllers/AnimNode_RigidBody.cpp
@@ -156,6 +156,7 @@ FAnimNode_RigidBody::FAnimNode_RigidBody()
 #if WITH_EDITORONLY_DATA
 	, bComponentSpaceSimulation_DEPRECATED(true)
 #endif
+	, SimulationTiming(ESimulationTiming::Default)
 	, WorldTimeSeconds(0.0f)
 	, LastEvalTimeSeconds(0.0f)
 	, AccumulatedDeltaTime(0.0f)
