[UE][FEATURE] RBAN Access to cloth colliders
- Rigid Body Animation Nodes can now add kinematic objects at runtime to represent any Cloth Colliders defined in a parent skeleton. This can be used to improve the behavior of  RBAN alternatives to cloth simulation (often used at lower LODs).
- There is a new Checkbox (called Use External Cloth Collision) on the RBAN node preferences panel to enable this feature .

diff --git a/Engine/Source/Runtime/AnimGraphRuntime/Private/BoneControllers/AnimNode_RigidBody.cpp b/Engine/Source/Runtime/AnimGraphRuntime/Private/BoneControllers/AnimNode_RigidBody.cpp
--- a/Engine/Source/Runtime/AnimGraphRuntime/Private/BoneControllers/AnimNode_RigidBody.cpp
+++ b/Engine/Source/Runtime/AnimGraphRuntime/Private/BoneControllers/AnimNode_RigidBody.cpp
@@ -51,6 +51,9 @@ FAutoConsoleVariableRef CVarRigidBodyNodeEnableTimeBasedReset(TEXT("p.RigidBodyN
 FAutoConsoleVariableRef CVarRigidBodyNodeEnableComponentAcceleration(TEXT("p.RigidBodyNode.EnableComponentAcceleration"), bRBAN_EnableComponentAcceleration, TEXT("Enable/Disable the simple acceleration transfer system for component- or bone-space simulation"), ECVF_Default);
 FAutoConsoleVariableRef CVarRigidBodyNodeWorldObjectExpiry(TEXT("p.RigidBodyNode.WorldObjectExpiry"), RBAN_WorldObjectExpiry, TEXT("World objects are removed from the simulation if not detected after this many tests"), ECVF_Default);
 
+bool bRBAN_IncludeClothColliders = true;
+FAutoConsoleVariableRef CVarRigidBodyNodeIncludeClothColliders(TEXT("p.RigidBodyNode.IncludeClothColliders"), bRBAN_IncludeClothColliders, TEXT("Include cloth colliders as kinematic bodies in the immediate physics simulation."), ECVF_Default);
+
 // FSimSpaceSettings forced overrides for testing
 bool bRBAN_SimSpace_EnableOverride = false;
 FSimSpaceSettings RBAN_SimSpaceOverride;
@@ -142,6 +145,7 @@ FAnimNode_RigidBody::FAnimNode_RigidBody()
 	, OverlapChannel(ECC_WorldStatic)
 	, SimulationSpace(ESimulationSpace::ComponentSpace)
 	, bForceDisableCollisionBetweenConstraintBodies(false)
+	, bUseExternalClothCollision(false)
 	, ResetSimulatedTeleportType(ETeleportType::None)
 	, bEnableWorldGeometry(false)
 	, bOverrideWorldGravity(false)
@@ -515,6 +519,7 @@ void FAnimNode_RigidBody::FlushDeferredSimulationTask()
 
 void FAnimNode_RigidBody::DestroyPhysicsSimulation()
 {
+	ClothColliders.Reset();
 	FlushDeferredSimulationTask();
 	delete PhysicsSimulation;
 	PhysicsSimulation = nullptr;
@@ -821,6 +826,7 @@ void FAnimNode_RigidBody::EvaluateSkeletalControl_AnyThread(FComponentSpacePoseC
 				SimulationAngularAcceleration);
 
 			UpdateWorldObjects(SimulationTransform);
+			UpdateClothColliderObjects(SimulationTransform);
 
 			PhysicsSimulation->UpdateSimulationSpace(
 				SimulationTransform, 
@@ -1262,6 +1268,8 @@ void FAnimNode_RigidBody::InitPhysics(const UAnimInstance* InAnimInstance)
 		PhysicsSimulation->SetIgnoreCollisionPairTable(IgnorePairs);
 		PhysicsSimulation->SetIgnoreCollisionActors(IgnoreCollisionActors);
 
+		CollectClothColliderObjects(SkeletalMeshComp);
+
 #if WITH_CHAOS
 		SolverSettings = UsePhysicsAsset->SolverSettings;
 		PhysicsSimulation->SetSolverSettings(
@@ -1560,6 +1568,88 @@ void FAnimNode_RigidBody::UpdateInternal(const FAnimationUpdateContext& Context)
 	PhysScene = nullptr;
 }
 
+void FAnimNode_RigidBody::CollectClothColliderObjects(const USkeletalMeshComponent* SkeletalMeshComp)
+{
+	if (bUseExternalClothCollision && bRBAN_IncludeClothColliders && SkeletalMeshComp && PhysicsSimulation)
+	{
+		// @todo(ccaulfield): is there an engine-independent way to do this?
+#if WITH_PHYSX && PHYSICS_INTERFACE_PHYSX
+		SCOPED_SCENE_READ_LOCK(PhysScene ? PhysScene->GetPxScene() : nullptr);
+#endif
+
+		const TArray<FClothCollisionSource>& SkeletalMeshClothCollisionSources = SkeletalMeshComp->GetClothCollisionSources();
+		
+		for (const FClothCollisionSource& ClothCollisionSource : SkeletalMeshClothCollisionSources)
+		{
+			const USkeletalMeshComponent* const SourceComponent = ClothCollisionSource.SourceComponent.Get();
+			const UPhysicsAsset* const PhysicsAsset = ClothCollisionSource.SourcePhysicsAsset.Get();
+
+			if (SourceComponent && PhysicsAsset)
+			{
+				TArray<FBodyInstance*> BodyInstances;
+				SourceComponent->InstantiatePhysicsAssetBodies(*PhysicsAsset, BodyInstances);
+
+				for (uint32 BodyInstanceIndex = 0, BodyInstanceMax = BodyInstances.Num(); BodyInstanceIndex < BodyInstanceMax; ++BodyInstanceIndex)
+				{
+					FBodyInstance* const BodyInstance = BodyInstances[BodyInstanceIndex];
+
+#if WITH_PHYSX && PHYSICS_INTERFACE_PHYSX
+					// Not sure why this happens, adding check to fix crash in CheckRBN engine test.
+					if (BodyInstance.BodySetup != nullptr)
+					{
+						ImmediatePhysics::FActorHandle* const ActorHandle = PhysicsSimulation->CreateActor(ImmediatePhysics::EActorType::StaticActor, BodyInstance, BodyInstance->GetUnrealWorldTransform());
+						ClothColliders.Add(FClothCollider(ActorHandle, SourceComponent, BodyInstance->InstanceBoneIndex));
+					}
+#elif WITH_CHAOS
+					ImmediatePhysics::FActorHandle* const ActorHandle = PhysicsSimulation->CreateActor(ImmediatePhysics::EActorType::KinematicActor, BodyInstance, BodyInstance->GetUnrealWorldTransform());
+					PhysicsSimulation->AddToCollidingPairs(ActorHandle); // <-allow collision between this actor and all dynamic actors.
+					ClothColliders.Add(FClothCollider(ActorHandle, SourceComponent, BodyInstance->InstanceBoneIndex));
+#endif
+					// Terminate the instance.
+					if (BodyInstance->IsValidBodyInstance())
+					{
+						BodyInstance->TermBody(true);
+					}
+
+					delete BodyInstance;
+					BodyInstances[BodyInstanceIndex] = nullptr;
+				}
+
+				BodyInstances.Reset();
+			}
+		}
+	}
+}
+
+void FAnimNode_RigidBody::RemoveClothColliderObjects()
+{
+	for (const FClothCollider& ClothCollider : ClothColliders)
+	{
+		PhysicsSimulation->DestroyActor(ClothCollider.ActorHandle);
+	}
+	
+	ClothColliders.Reset();
+}
+
+void FAnimNode_RigidBody::UpdateClothColliderObjects(const FTransform& SpaceTransform)
+{
+	for (FClothCollider& ClothCollider : ClothColliders)
+	{
+		if (ClothCollider.ActorHandle && ClothCollider.SkeletalMeshComponent)
+		{
+			// Calculate the sim-space transform of this object
+			const FTransform CompWorldTransform = ClothCollider.SkeletalMeshComponent->GetBoneTransform(ClothCollider.BoneIndex);
+			FTransform CompSpaceTransform;
+			CompSpaceTransform.SetTranslation(SpaceTransform.InverseTransformPosition(CompWorldTransform.GetLocation()));
+			CompSpaceTransform.SetRotation(SpaceTransform.InverseTransformRotation(CompWorldTransform.GetRotation()));
+			CompSpaceTransform.SetScale3D(FVector::OneVector);	// TODO - sort out scale for world objects in local sim
+
+			// Update the sim's copy of the world object
+			ClothCollider.ActorHandle->SetWorldTransform(CompSpaceTransform);
+		}
+	}
+}
+
 void FAnimNode_RigidBody::CollectWorldObjects()
 {
 	if ((UnsafeWorld != nullptr) && (PhysScene != nullptr))
diff --git a/Engine/Source/Runtime/AnimGraphRuntime/Public/BoneControllers/AnimNode_RigidBody.h b/Engine/Source/Runtime/AnimGraphRuntime/Public/BoneControllers/AnimNode_RigidBody.h
--- a/Engine/Source/Runtime/AnimGraphRuntime/Public/BoneControllers/AnimNode_RigidBody.h
+++ b/Engine/Source/Runtime/AnimGraphRuntime/Public/BoneControllers/AnimNode_RigidBody.h
@@ -231,6 +231,9 @@ public:
 	UPROPERTY(EditAnywhere, Category = Settings)
 	bool bForceDisableCollisionBetweenConstraintBodies;
 
+	/** If true, kinematic objects will be added to the simulation at runtime to represent any cloth colliders defined for the parent object. */
+	UPROPERTY(EditAnywhere, Category = Settings)
+	bool bUseExternalClothCollision;
 
 private:
 	ETeleportType ResetSimulatedTeleportType;
@@ -319,6 +322,15 @@ private:
 		FVector& SpaceLinearAcc,
 		FVector& SpaceAngularAcc);
 
+	// Gather cloth collision sources from the supplied Skeltal Mesh and add a kinematic actor representing each one of them to the sim.
+	void CollectClothColliderObjects(const USkeletalMeshComponent* SkeletalMeshComp);
+	
+	// Remove all cloth collider objects from the sim.
+	void RemoveClothColliderObjects();
+
+	// Update the sim-space transforms of all cloth collider objects.
+	void UpdateClothColliderObjects(const FTransform& SpaceTransform);
+
 	// Gather nearby world objects and add them to the sim
 	void CollectWorldObjects();
 
@@ -427,6 +439,23 @@ private:
 
 	FPerSolverFieldSystem PerSolverField;
 
+	// Information required to identify and update a kinematic object representing a cloth collision source in the sim.
+	struct FClothCollider
+	{
+		FClothCollider(ImmediatePhysics::FActorHandle* const InActorHandle, const USkeletalMeshComponent* const InSkeletalMeshComponent, const uint32 InBoneIndex)
+			: ActorHandle(InActorHandle)
+			, SkeletalMeshComponent(InSkeletalMeshComponent)
+			, BoneIndex(InBoneIndex)
+		{}
+
+		ImmediatePhysics::FActorHandle* ActorHandle; // Identifies the physics actor in the sim.
+		const USkeletalMeshComponent* SkeletalMeshComponent; // Parent skeleton.
+		uint32 BoneIndex; // Bone within parent skeleton that drives physics actors transform.
+	};
+
+	// List of actors in the sim that represent objects collected from other parts of this character.
+	TArray<FClothCollider> ClothColliders; 
+	
 	TMap<const UPrimitiveComponent*, FWorldObject> ComponentsInSim;
 	int32 ComponentsInSimTick;
 
diff --git a/Engine/Source/Runtime/Engine/Classes/Components/SkeletalMeshComponent.h b/Engine/Source/Runtime/Engine/Classes/Components/SkeletalMeshComponent.h
--- a/Engine/Source/Runtime/Engine/Classes/Components/SkeletalMeshComponent.h
+++ b/Engine/Source/Runtime/Engine/Classes/Components/SkeletalMeshComponent.h
@@ -302,6 +302,35 @@ struct ENGINE_API FClosestPointOnPhysicsAsset
 	}
 };
 
+/** Helper struct used to store info about a cloth collision source */
+struct FClothCollisionSource
+{
+	FClothCollisionSource(USkeletalMeshComponent* InSourceComponent, UPhysicsAsset* InSourcePhysicsAsset, const FOnBoneTransformsFinalizedMultiCast::FDelegate& InOnBoneTransformsFinalizedDelegate);
+	~FClothCollisionSource();
+
+	/** Component that collision data will be copied from */
+	TWeakObjectPtr<USkeletalMeshComponent> SourceComponent;
+
+	/** Physics asset to use to generate collision against the source component */
+	TWeakObjectPtr<UPhysicsAsset> SourcePhysicsAsset;
+
+	/** Callback used to remove the cloth transform updates delegate */
+	FDelegateHandle OnBoneTransformsFinalizedHandle;
+
+	/** Cached skeletal mesh used to invalidate the cache if the skeletal mesh has changed */
+	TWeakObjectPtr<USkeletalMesh> CachedSkeletalMesh;
+
+	/** Cached spheres from physics asset */
+	TArray<FClothCollisionPrim_Sphere> CachedSpheres;
+
+	/** Cached sphere connections from physics asset */
+	TArray<FClothCollisionPrim_SphereConnection> CachedSphereConnections;
+
+	/** Flag whether the cache is valid */
+	bool bCached;
+};
+
+
 /**
  * SkeletalMeshComponent is used to create an instance of an animated SkeletalMesh asset.
  *
@@ -1439,34 +1468,6 @@ private:
 	UPROPERTY(Transient)
 	TObjectPtr<UClothingSimulationInteractor> ClothingInteractor;
 
-	/** Helper struct used to store info about a cloth collision source */
-	struct FClothCollisionSource
-	{
-		FClothCollisionSource(USkeletalMeshComponent* InSourceComponent, UPhysicsAsset* InSourcePhysicsAsset, const FOnBoneTransformsFinalizedMultiCast::FDelegate& InOnBoneTransformsFinalizedDelegate);
-		~FClothCollisionSource();
-
-		/** Component that collision data will be copied from */
-		TWeakObjectPtr<USkeletalMeshComponent> SourceComponent;
-
-		/** Physics asset to use to generate collision against the source component */
-		TWeakObjectPtr<UPhysicsAsset> SourcePhysicsAsset;
-
-		/** Callback used to remove the cloth transform updates delegate */
-		FDelegateHandle OnBoneTransformsFinalizedHandle;
-
-		/** Cached skeletal mesh used to invalidate the cache if the skeletal mesh has changed */
-		TWeakObjectPtr<USkeletalMesh> CachedSkeletalMesh;
-
-		/** Cached spheres from physics asset */
-		TArray<FClothCollisionPrim_Sphere> CachedSpheres;
-
-		/** Cached sphere connections from physics asset */
-		TArray<FClothCollisionPrim_SphereConnection> CachedSphereConnections;
-
-		/** Flag whether the cache is valid */
-		bool bCached;
-	};
-
 	/** Array of sources for cloth collision */
 	TArray<FClothCollisionSource> ClothCollisionSources;
 
@@ -1945,6 +1946,9 @@ public:
 	/** Instantiates bodies given a physics asset. Typically you should call InitArticulated unless you are planning to do something special with the bodies. The Created bodies and constraints are owned by the calling code and must be freed when necessary.*/
 	void InstantiatePhysicsAsset(const UPhysicsAsset& PhysAsset, const FVector& Scale3D, TArray<FBodyInstance*>& OutBodies, TArray<FConstraintInstance*>& OutConstraints, FPhysScene* PhysScene = nullptr, USkeletalMeshComponent* OwningComponent = nullptr, int32 UseRootBodyIndex = INDEX_NONE, const FPhysicsAggregateHandle& UseAggregate = FPhysicsAggregateHandle()) const;
 
+	/** Instantiates only the bodies given a physics asset, not the constraints. The Created bodies are owned by the calling code and must be freed when necessary.*/
+	void InstantiatePhysicsAssetBodies(const UPhysicsAsset& PhysAsset, TArray<FBodyInstance*>& OutBodies, FPhysScene* PhysScene = nullptr, USkeletalMeshComponent* OwningComponent = nullptr, int32 UseRootBodyIndex = INDEX_NONE, const FPhysicsAggregateHandle& UseAggregate = FPhysicsAggregateHandle()) const;
+
 	/** Instantiates bodies given a physics asset like InstantiatePhysicsAsset but instead of reading the current component state, this reads the ref-pose from the reference skeleton of the mesh. Useful if trying to create bodies to be used during any evaluation work */
 	void InstantiatePhysicsAssetRefPose(const UPhysicsAsset& PhysAsset, const FVector& Scale3D, TArray<FBodyInstance*>& OutBodies, TArray<FConstraintInstance*>& OutConstraints, FPhysScene* PhysScene = nullptr, USkeletalMeshComponent* OwningComponent = nullptr, int32 UseRootBodyIndex = INDEX_NONE, const FPhysicsAggregateHandle& UseAggregate = FPhysicsAggregateHandle(), bool bCreateBodiesInRefPose = false) const;
 
@@ -2208,6 +2212,8 @@ public:
 	/** if the vertex index is valid for simulated vertices, returns the position in world space */
 	bool GetClothSimulatedPosition_GameThread(const FGuid& AssetGuid, int32 VertexIndex, FVector& OutSimulPos) const;
 
+	const TArray<FClothCollisionSource>& GetClothCollisionSources() const { return ClothCollisionSources; }
+
 #if WITH_CLOTH_COLLISION_DETECTION
 
 	/**
@@ -2311,6 +2317,7 @@ private:
 	void GetWindForCloth_GameThread(FVector& WindVector, float& WindAdaption) const;
 
 	void InstantiatePhysicsAsset_Internal(const UPhysicsAsset& PhysAsset, const FVector& Scale3D, TArray<FBodyInstance*>& OutBodies, TArray<FConstraintInstance*>& OutConstraints, TFunctionRef<FTransform(int32)> BoneTransformGetter, FPhysScene* PhysScene = nullptr, USkeletalMeshComponent* OwningComponent = nullptr, int32 UseRootBodyIndex = INDEX_NONE, const FPhysicsAggregateHandle& UseAggregate = FPhysicsAggregateHandle()) const;
+	void InstantiatePhysicsAssetBodies_Internal(const UPhysicsAsset& PhysAsset, TArray<FBodyInstance*>& OutBodies, TFunctionRef<FTransform(int32)> BoneTransformGetter, TMap<FName, FBodyInstance*>* OutNameToBodyMap = nullptr, FPhysScene* PhysScene = nullptr, USkeletalMeshComponent* OwningComponent = nullptr, int32 UseRootBodyIndex = INDEX_NONE, const FPhysicsAggregateHandle& UseAggregate = FPhysicsAggregateHandle()) const;
 
 	// Reference to our current parallel animation evaluation task (if there is one)
 	FGraphEventRef				ParallelAnimationEvaluationTask;
diff --git a/Engine/Source/Runtime/Engine/Private/SkeletalMeshComponentPhysics.cpp b/Engine/Source/Runtime/Engine/Private/SkeletalMeshComponentPhysics.cpp
--- a/Engine/Source/Runtime/Engine/Private/SkeletalMeshComponentPhysics.cpp
+++ b/Engine/Source/Runtime/Engine/Private/SkeletalMeshComponentPhysics.cpp
@@ -147,7 +147,7 @@ FName FSkeletalMeshComponentEndPhysicsTickFunction::DiagnosticContext(bool bDeta
 	return FName(TEXT("SkeletalMeshComponentEndPhysicsTick"));
 }
 
-USkeletalMeshComponent::FClothCollisionSource::FClothCollisionSource(USkeletalMeshComponent* InSourceComponent, UPhysicsAsset* InSourcePhysicsAsset, const FOnBoneTransformsFinalizedMultiCast::FDelegate& InOnBoneTransformsFinalizedDelegate)
+FClothCollisionSource::FClothCollisionSource(USkeletalMeshComponent* InSourceComponent, UPhysicsAsset* InSourcePhysicsAsset, const FOnBoneTransformsFinalizedMultiCast::FDelegate& InOnBoneTransformsFinalizedDelegate)
 	: SourceComponent(InSourceComponent)
 	, SourcePhysicsAsset(InSourcePhysicsAsset)
 	, bCached(false)
@@ -158,7 +158,7 @@ USkeletalMeshComponent::FClothCollisionSource::FClothCollisionSource(USkeletalMe
 	}
 }
 
-USkeletalMeshComponent::FClothCollisionSource::~FClothCollisionSource()
+FClothCollisionSource::~FClothCollisionSource()
 {
 	if (SourceComponent.IsValid() && OnBoneTransformsFinalizedHandle.IsValid())
 	{
@@ -785,6 +785,16 @@ void USkeletalMeshComponent::InstantiatePhysicsAsset(const UPhysicsAsset& PhysAs
 	InstantiatePhysicsAsset_Internal(PhysAsset, Scale3D, OutBodies, OutConstraints, BoneTMCallable, PhysScene, OwningComponent, UseRootBodyIndex, UseAggregate);
 }
 
+void USkeletalMeshComponent::InstantiatePhysicsAssetBodies(const UPhysicsAsset& PhysAsset, TArray<FBodyInstance*>& OutBodies, FPhysScene* PhysScene /*= nullptr*/, USkeletalMeshComponent* OwningComponent /*= nullptr*/, int32 UseRootBodyIndex /*= INDEX_NONE*/, const FPhysicsAggregateHandle& UseAggregate /*= FPhysicsAggregateHandle()*/) const
+{
+	auto BoneTMCallable = [this](int32 BoneIndex)
+	{
+		return GetBoneTransform(BoneIndex);
+	};
+
+	InstantiatePhysicsAssetBodies_Internal(PhysAsset, OutBodies, BoneTMCallable, nullptr, PhysScene, OwningComponent, UseRootBodyIndex, UseAggregate);
+}
+
 void USkeletalMeshComponent::InstantiatePhysicsAssetRefPose(const UPhysicsAsset& PhysAsset, const FVector& Scale3D, TArray<FBodyInstance*>& OutBodies, TArray<FConstraintInstance*>& OutConstraints, FPhysScene* PhysScene /*= nullptr*/, USkeletalMeshComponent* OwningComponent /*= nullptr*/, int32 UseRootBodyIndex /*= INDEX_NONE*/, const FPhysicsAggregateHandle& UseAggregate, bool bCreateBodiesInRefPose) const
 {
 	if(SkeletalMesh)
@@ -823,18 +833,18 @@ void USkeletalMeshComponent::InstantiatePhysicsAssetRefPose(const UPhysicsAsset&
 	}
 }
 
-void USkeletalMeshComponent::InstantiatePhysicsAsset_Internal(const UPhysicsAsset& PhysAsset, const FVector& Scale3D, TArray<FBodyInstance*>& OutBodies, TArray<FConstraintInstance*>& OutConstraints, TFunctionRef<FTransform(int32)> BoneTransformGetter, FPhysScene* PhysScene /*= nullptr*/, USkeletalMeshComponent* OwningComponent /*= nullptr*/, int32 UseRootBodyIndex /*= INDEX_NONE*/, const FPhysicsAggregateHandle& UseAggregate) const
+void USkeletalMeshComponent::InstantiatePhysicsAssetBodies_Internal(const UPhysicsAsset& PhysAsset, TArray<FBodyInstance*>& OutBodies, TFunctionRef<FTransform(int32)> BoneTransformGetter, TMap<FName, FBodyInstance*>* OutNameToBodyMap, FPhysScene* PhysScene /*= nullptr*/, USkeletalMeshComponent* OwningComponent /*= nullptr*/, int32 UseRootBodyIndex /*= INDEX_NONE*/, const FPhysicsAggregateHandle& UseAggregate) const
 {
-	const float ActualScale = Scale3D.GetAbsMin();
-	const float Scale = ActualScale == 0.f ? KINDA_SMALL_NUMBER : ActualScale;
-
 	const int32 NumOutBodies = PhysAsset.SkeletalBodySetups.Num();
 
-	TMap<FName, FBodyInstance*> NameToBodyMap;
-
 	// Create all the OutBodies.
 	check(OutBodies.Num() == 0);
 	OutBodies.AddZeroed(NumOutBodies);
+	
+	if (OutNameToBodyMap)
+	{
+		OutNameToBodyMap->Reserve(NumOutBodies);
+	}
 
 	for(int32 BodyIdx = 0; BodyIdx < NumOutBodies; BodyIdx++)
 	{
@@ -904,9 +914,23 @@ void USkeletalMeshComponent::InstantiatePhysicsAsset_Internal(const UPhysicsAsse
 			SpawnParams.Aggregate = UseAggregate;
 			BodyInst->InitBody(PhysicsAssetBodySetup, BoneTransform, OwningComponent, PhysScene, SpawnParams);
 
-			NameToBodyMap.Add(PhysicsAssetBodySetup->BoneName, BodyInst);
+			if (OutNameToBodyMap)
+			{
+				OutNameToBodyMap->Add(PhysicsAssetBodySetup->BoneName, BodyInst);
+			}
 		}
 	}
+}
+
+void USkeletalMeshComponent::InstantiatePhysicsAsset_Internal(const UPhysicsAsset& PhysAsset, const FVector& Scale3D, TArray<FBodyInstance*>& OutBodies, TArray<FConstraintInstance*>& OutConstraints, TFunctionRef<FTransform(int32)> BoneTransformGetter, FPhysScene* PhysScene /*= nullptr*/, USkeletalMeshComponent* OwningComponent /*= nullptr*/, int32 UseRootBodyIndex /*= INDEX_NONE*/, const FPhysicsAggregateHandle& UseAggregate) const
+{
+	const float ActualScale = Scale3D.GetAbsMin();
+	const float Scale = ActualScale == 0.f ? KINDA_SMALL_NUMBER : ActualScale;
+
+	TMap<FName, FBodyInstance*> NameToBodyMap;
+
+	// Create all the OutBodies.
+	InstantiatePhysicsAssetBodies_Internal(PhysAsset, OutBodies, BoneTransformGetter, &NameToBodyMap, PhysScene, OwningComponent, UseRootBodyIndex, UseAggregate);
 
 	if(PhysScene && Aggregate.IsValid())
 	{
