Fix for LayeredBoneBlend improper handling of invalid curve weight values

diff --git a/Engine/Source/Runtime/AnimGraphRuntime/Private/AnimNodes/AnimNode_LayeredBoneBlend.cpp b/Engine/Source/Runtime/AnimGraphRuntime/Private/AnimNodes/AnimNode_LayeredBoneBlend.cpp
--- a/Engine/Source/Runtime/AnimGraphRuntime/Private/AnimNodes/AnimNode_LayeredBoneBlend.cpp
+++ b/Engine/Source/Runtime/AnimGraphRuntime/Private/AnimNodes/AnimNode_LayeredBoneBlend.cpp
@@ -259,14 +259,14 @@ void FAnimNode_LayeredBoneBlend::Evaluate_AnyThread(FPoseContext& Output)
 				int32 SourceIndex = CurvePoseSourceIndices[CurvePoseIndex];
 				if (SourceIndex != DEFAULT_SOURCEINDEX)
 				{
-					// if source index is set, clear base pose curve value
-					BasePoseContext.Curve.Set(UIDIndex, 0.f);
+					// if source index is set, invalidate base pose curve value
+					BasePoseContext.Curve.InvalidateCurveWeight(UIDIndex);
 					for (int32 ChildIndex = 0; ChildIndex < NumPoses; ++ChildIndex)
 					{
 						if (SourceIndex != ChildIndex)
 						{
-							// if not source, clear it
-							TargetBlendCurves[ChildIndex].Set(UIDIndex, 0.f);
+							// if not source, invalidate it
+							TargetBlendCurves[ChildIndex].InvalidateCurveWeight(UIDIndex);
 						}
 					}
 				}
diff --git a/Engine/Source/Runtime/Engine/Public/Animation/AnimCurveTypes.h b/Engine/Source/Runtime/Engine/Public/Animation/AnimCurveTypes.h
--- a/Engine/Source/Runtime/Engine/Public/Animation/AnimCurveTypes.h
+++ b/Engine/Source/Runtime/Engine/Public/Animation/AnimCurveTypes.h
@@ -401,6 +401,19 @@ struct FBaseBlendedCurve
 		}
 	}
 
+	/** Invalidate value of InUID */
+	void InvalidateCurveWeight(USkeleton::AnimCurveUID InUid)
+	{
+		check(bInitialized);
+
+		int32 ArrayIndex = GetArrayIndexByUID(InUid);
+		if (ArrayIndex != INDEX_NONE)
+		{
+			CurveWeights[ArrayIndex] = 0.f;
+			ValidCurveWeights[ArrayIndex] = false;
+		}
+	}
+
 	/** Set value of InUID to InValue */
 	void Set(USkeleton::AnimCurveUID InUid, float InValue)
 	{
