Fixed crash compiling a template animation blueprint with a layered blend per bone node in it

diff --git a/Engine/Source/Runtime/AnimGraphRuntime/Private/AnimNodes/AnimNode_LayeredBoneBlend.cpp b/Engine/Source/Runtime/AnimGraphRuntime/Private/AnimNodes/AnimNode_LayeredBoneBlend.cpp
--- a/Engine/Source/Runtime/AnimGraphRuntime/Private/AnimNodes/AnimNode_LayeredBoneBlend.cpp
+++ b/Engine/Source/Runtime/AnimGraphRuntime/Private/AnimNodes/AnimNode_LayeredBoneBlend.cpp
@@ -49,7 +49,7 @@ void FAnimNode_LayeredBoneBlend::RebuildPerBoneBlendWeights(const USkeleton* InS
 
 bool FAnimNode_LayeredBoneBlend::ArePerBoneBlendWeightsValid(const USkeleton* InSkeleton) const
 {
-	return (InSkeleton->GetGuid() == SkeletonGuid && InSkeleton->GetVirtualBoneGuid() == VirtualBoneGuid);
+	return (InSkeleton != nullptr && InSkeleton->GetGuid() == SkeletonGuid && InSkeleton->GetVirtualBoneGuid() == VirtualBoneGuid);
 }
 
 void FAnimNode_LayeredBoneBlend::UpdateCachedBoneData(const FBoneContainer& RequiredBones, const USkeleton* Skeleton)
