PoseSearch: Don't reset asset time when switching blend spaces in motion matching

diff --git a/Engine/Plugins/Experimental/Animation/PoseSearch/Source/Runtime/Private/AnimNode_MotionMatching.cpp b/Engine/Plugins/Experimental/Animation/PoseSearch/Source/Runtime/Private/AnimNode_MotionMatching.cpp
--- a/Engine/Plugins/Experimental/Animation/PoseSearch/Source/Runtime/Private/AnimNode_MotionMatching.cpp
+++ b/Engine/Plugins/Experimental/Animation/PoseSearch/Source/Runtime/Private/AnimNode_MotionMatching.cpp
@@ -33,6 +33,8 @@ void FAnimNode_MotionMatching::Initialize_AnyThread(const FAnimationInitializeCo
 		MirrorNode.SetMirrorDataTable(Database->Schema->MirrorDataTable.Get());
 	}
 
+	BlendSpacePlayerNode.SetResetPlayTimeWhenBlendSpaceChanges(false /*!bReset*/);
+
 	Source.SetLinkNode(&MirrorNode);
 	Source.Initialize(Context);
 }
diff --git a/Engine/Source/Runtime/AnimGraphRuntime/Private/AnimNodes/AnimNode_BlendSpacePlayer.cpp b/Engine/Source/Runtime/AnimGraphRuntime/Private/AnimNodes/AnimNode_BlendSpacePlayer.cpp
--- a/Engine/Source/Runtime/AnimGraphRuntime/Private/AnimNodes/AnimNode_BlendSpacePlayer.cpp
+++ b/Engine/Source/Runtime/AnimGraphRuntime/Private/AnimNodes/AnimNode_BlendSpacePlayer.cpp
@@ -281,6 +281,20 @@ bool FAnimNode_BlendSpacePlayer::SetIgnoreForRelevancyTest(bool bInIgnoreForRele
 	return false;
 }
 
+bool FAnimNode_BlendSpacePlayer::SetResetPlayTimeWhenBlendSpaceChanges(bool bInReset)
+{
+#if WITH_EDITORONLY_DATA
+	bResetPlayTimeWhenBlendSpaceChanges = bInReset;
+#endif
+
+	if (bool* bResetPlayTimeWhenBlendSpaceChangesPtr = GET_INSTANCE_ANIM_NODE_DATA_PTR(bool, bResetPlayTimeWhenBlendSpaceChanges))
+	{
+		*bResetPlayTimeWhenBlendSpaceChangesPtr = bInReset;
+		return true;
+	}
+
+	return false;
+}
 
 bool FAnimNode_BlendSpacePlayer::SetBlendSpace(UBlendSpace* InBlendSpace)
 {
diff --git a/Engine/Source/Runtime/AnimGraphRuntime/Public/AnimNodes/AnimNode_BlendSpacePlayer.h b/Engine/Source/Runtime/AnimGraphRuntime/Public/AnimNodes/AnimNode_BlendSpacePlayer.h
--- a/Engine/Source/Runtime/AnimGraphRuntime/Public/AnimNodes/AnimNode_BlendSpacePlayer.h
+++ b/Engine/Source/Runtime/AnimGraphRuntime/Public/AnimNodes/AnimNode_BlendSpacePlayer.h
@@ -72,6 +72,9 @@ public:
 	// Get whether we should reset the current play time when the blend space changes
 	virtual bool ShouldResetPlayTimeWhenBlendSpaceChanges() const PURE_VIRTUAL(FAnimNode_BlendSpacePlayerBase::ShouldResetPlayTimeWhenBlendSpaceChanges, return true;);
 
+	// Set whether we should reset the current play time when the blend space changes
+	virtual bool SetResetPlayTimeWhenBlendSpaceChanges(bool bReset) PURE_VIRTUAL(FAnimNode_BlendSpacePlayerBase::SetResetPlayTimeWhenBlendSpaceChanges, return false;);
+
 	// An evaluator will be setting the play rate to zero and setting the time explicitly. ShouldTeleportToTime indicates whether we should jump to that time, or move to it playing out root motion and events etc.
 	virtual bool ShouldTeleportToTime() const { return false; }
 
@@ -188,6 +191,7 @@ public:
 	virtual float GetPlayRate() const override;
 	virtual bool GetLoop() const override;
 	virtual bool ShouldResetPlayTimeWhenBlendSpaceChanges() const override;
+	virtual bool SetResetPlayTimeWhenBlendSpaceChanges(bool bReset) override;
 	virtual bool SetBlendSpace(UBlendSpace* InBlendSpace) override;
 	virtual bool SetPosition(FVector InPosition) override;
 	virtual bool SetPlayRate(float InPlayRate) override;
@@ -267,6 +271,7 @@ public:
 	virtual float GetPlayRate() const override { return PlayRate; }
 	virtual bool GetLoop() const override { return bLoop; }
 	virtual bool ShouldResetPlayTimeWhenBlendSpaceChanges() const override { return bResetPlayTimeWhenBlendSpaceChanges; }
+	virtual bool SetResetPlayTimeWhenBlendSpaceChanges(bool bReset) override { bResetPlayTimeWhenBlendSpaceChanges = bReset; return true; }
 	virtual bool SetBlendSpace(UBlendSpace* InBlendSpace) override { BlendSpace = InBlendSpace; return true; }
 	virtual bool SetPosition(FVector InPosition) override { X = InPosition[0]; Y = InPosition[1]; return true; }
 	virtual bool SetPlayRate(float InPlayRate) override { PlayRate = InPlayRate; return true; }
