Anim: Remove stray PRAGMA_DISABLE_OPTIMIZATION

diff --git a/Engine/Source/Runtime/AnimGraphRuntime/Public/AnimCustomInstanceHelper.h b/Engine/Source/Runtime/AnimGraphRuntime/Public/AnimCustomInstanceHelper.h
--- a/Engine/Source/Runtime/AnimGraphRuntime/Public/AnimCustomInstanceHelper.h
+++ b/Engine/Source/Runtime/AnimGraphRuntime/Public/AnimCustomInstanceHelper.h
@@ -30,8 +30,6 @@
 #include "AnimSequencerInstance.h"
 #include "SequencerAnimationSupport.h"
 
-
-PRAGMA_DISABLE_OPTIMIZATION
 class ANIMGRAPHRUNTIME_API FAnimCustomInstanceHelper
 {
 public:
@@ -288,5 +286,4 @@ public:
 private:
 	/** Helper function for BindToSkeletalMeshComponent */
 	static bool ShouldCreateCustomInstancePlayer(const USkeletalMeshComponent* SkeletalMeshComponent);
-};
-PRAGMA_ENABLE_OPTIMIZATION
\ No newline at end of file
+};
\ No newline at end of file
