Sequencer: Swap Root Bone with Component Root, Actor Root to put the root motion on the skeletal mesh component or the actor's root component.

diff --git a/Engine/Source/Runtime/AnimGraphRuntime/Private/AnimSequencerInstance.cpp b/Engine/Source/Runtime/AnimGraphRuntime/Private/AnimSequencerInstance.cpp
--- a/Engine/Source/Runtime/AnimGraphRuntime/Private/AnimSequencerInstance.cpp
+++ b/Engine/Source/Runtime/AnimGraphRuntime/Private/AnimSequencerInstance.cpp
@@ -37,7 +37,7 @@ void UAnimSequencerInstance::UpdateAnimTrack(UAnimSequenceBase* InAnimSequence,
 
 void UAnimSequencerInstance::UpdateAnimTrackWithRootMotion(UAnimSequenceBase* InAnimSequence, int32 SequenceId, const TOptional<FRootMotionOverride>& RootMotion, float InFromPosition, float InToPosition, float Weight, bool bFireNotifies)
 {
-	FAnimSequencerData AnimSequencerData(InAnimSequence, SequenceId, RootMotion, InFromPosition, InToPosition, Weight, bFireNotifies, false, FTransform(), nullptr);
+	FAnimSequencerData AnimSequencerData(InAnimSequence, SequenceId, RootMotion, InFromPosition, InToPosition, Weight, bFireNotifies, ESwapRootBone::SwapRootBone_None, FTransform(), nullptr);
 	GetProxyOnGameThread<FAnimSequencerInstanceProxy>().UpdateAnimTrackWithRootMotion(AnimSequencerData);
 }
 
diff --git a/Engine/Source/Runtime/AnimGraphRuntime/Private/AnimSequencerInstanceProxy.cpp b/Engine/Source/Runtime/AnimGraphRuntime/Private/AnimSequencerInstanceProxy.cpp
--- a/Engine/Source/Runtime/AnimGraphRuntime/Private/AnimSequencerInstanceProxy.cpp
+++ b/Engine/Source/Runtime/AnimGraphRuntime/Private/AnimSequencerInstanceProxy.cpp
@@ -63,7 +63,7 @@ bool FAnimSequencerInstanceProxy::Evaluate(FPoseContext& Output)
 	
 	RootBoneTransform.Reset();
 
-	if (bSwapRootBoneWithComponentRoot)
+	if (SwapRootBone != ESwapRootBone::SwapRootBone_None)
 	{
 		for (const FCompactPoseBoneIndex BoneIndex : Output.Pose.ForEachBoneIndex())
 		{
@@ -81,7 +81,7 @@ bool FAnimSequencerInstanceProxy::Evaluate(FPoseContext& Output)
 
 void FAnimSequencerInstanceProxy::PostEvaluate(UAnimInstance* InAnimInstance)
 {
-	if (GetSkelMeshComponent() && bSwapRootBoneWithComponentRoot)
+	if (GetSkelMeshComponent() && SwapRootBone != ESwapRootBone::SwapRootBone_None)
 	{
 		if (RootBoneTransform.IsSet())
 		{
@@ -91,7 +91,19 @@ void FAnimSequencerInstanceProxy::PostEvaluate(UAnimInstance* InAnimInstance)
 			{
 				RelativeTransform = RootBoneTransform.GetValue() * InitialTransform.GetValue();
 			}
-			GetSkelMeshComponent()->SetRelativeLocationAndRotation(RelativeTransform.GetLocation(), RelativeTransform.GetRotation().Rotator());
+
+			if (SwapRootBone == ESwapRootBone::SwapRootBone_Component)
+			{
+				GetSkelMeshComponent()->SetRelativeLocationAndRotation(RelativeTransform.GetLocation(), RelativeTransform.GetRotation().Rotator());
+			}
+			else if (SwapRootBone == ESwapRootBone::SwapRootBone_Actor)
+			{
+				AActor* Actor = GetSkelMeshComponent()->GetOwner();
+				if (Actor && Actor->GetRootComponent())
+				{
+					Actor->GetRootComponent()->SetRelativeLocationAndRotation(RelativeTransform.GetLocation(), RelativeTransform.GetRotation().Rotator());
+				}
+			}
 		}
 	}
 }
@@ -246,7 +258,7 @@ void FAnimSequencerInstanceProxy::UpdateAnimTrackWithRootMotion(UAnimSequenceBas
 
 void FAnimSequencerInstanceProxy::UpdateAnimTrackWithRootMotion(const FAnimSequencerData& InAnimSequencerData)
 {
-	bSwapRootBoneWithComponentRoot = InAnimSequencerData.bSwapRootBoneWithComponentRoot;
+	SwapRootBone = InAnimSequencerData.SwapRootBone;
 	InitialTransform = InAnimSequencerData.InitialTransform;
 	UpdateAnimTrack(InAnimSequencerData.AnimSequence, InAnimSequencerData.SequenceId, InAnimSequencerData.RootMotion, InAnimSequencerData.FromPosition, InAnimSequencerData.ToPosition, InAnimSequencerData.Weight, InAnimSequencerData.bFireNotifies, InAnimSequencerData.MirrorDataTable);
 }
diff --git a/Engine/Source/Runtime/AnimGraphRuntime/Public/AnimSequencerInstanceProxy.h b/Engine/Source/Runtime/AnimGraphRuntime/Public/AnimSequencerInstanceProxy.h
--- a/Engine/Source/Runtime/AnimGraphRuntime/Public/AnimSequencerInstanceProxy.h
+++ b/Engine/Source/Runtime/AnimGraphRuntime/Public/AnimSequencerInstanceProxy.h
@@ -10,6 +10,17 @@
 #include "AnimNodes/AnimNode_Mirror.h"
 #include "AnimSequencerInstanceProxy.generated.h"
 
+UENUM()
+enum class ESwapRootBone : uint8
+{
+	/* Swap the root bone to the component */
+	SwapRootBone_Component,
+	/* Swap root bone to the actor root component */
+	SwapRootBone_Actor,
+	/* Do not swap the root bone */
+	SwapRootBone_None
+};
+
 /** Base class for all 'players' that can attach to and be blended into a sequencer instance's output */
 struct FSequencerPlayerBase
 {
@@ -58,7 +69,7 @@ struct FRootMotionOverride
 
 struct ANIMGRAPHRUNTIME_API FAnimSequencerData
 {
-	FAnimSequencerData(UAnimSequenceBase* InAnimSequence, int32 InSequenceId, const TOptional<FRootMotionOverride>& InRootMotion, float InFromPosition, float InToPosition, float InWeight, bool bInFireNotifies, bool bInSwapRootBoneWithComponentRoot, TOptional<FTransform> InInitialTransform, UMirrorDataTable* InMirrorDataTable)
+	FAnimSequencerData(UAnimSequenceBase* InAnimSequence, int32 InSequenceId, const TOptional<FRootMotionOverride>& InRootMotion, float InFromPosition, float InToPosition, float InWeight, bool bInFireNotifies, ESwapRootBone InSwapRootBone, TOptional<FTransform> InInitialTransform, UMirrorDataTable* InMirrorDataTable)
 		: AnimSequence(InAnimSequence)
 		, SequenceId(InSequenceId)
 		, RootMotion(InRootMotion)
@@ -66,7 +77,7 @@ struct ANIMGRAPHRUNTIME_API FAnimSequencerData
 		, ToPosition(InToPosition)
 		, Weight(InWeight)
 		, bFireNotifies(bInFireNotifies)
-		, bSwapRootBoneWithComponentRoot(bInSwapRootBoneWithComponentRoot)
+		, SwapRootBone(InSwapRootBone)
 		, InitialTransform(InInitialTransform)
 		, MirrorDataTable(InMirrorDataTable)
 	{
@@ -79,7 +90,7 @@ struct ANIMGRAPHRUNTIME_API FAnimSequencerData
 	float ToPosition;
 	float Weight;
 	bool bFireNotifies;
-	bool bSwapRootBoneWithComponentRoot;
+	ESwapRootBone SwapRootBone;
 	TOptional<FTransform> InitialTransform;
 	UMirrorDataTable* MirrorDataTable;
 };
@@ -179,7 +190,7 @@ protected:
 	void EnsureAnimTrack(UAnimSequenceBase* InAnimSequence, uint32 SequenceId);
 	void ClearSequencePlayerAndMirrorMaps();
 
-	bool bSwapRootBoneWithComponentRoot = false;
+	ESwapRootBone SwapRootBone = ESwapRootBone::SwapRootBone_None;
 	TOptional<FTransform> InitialTransform;
 	TOptional<FTransform> RootBoneTransform;
 };
diff --git a/Engine/Source/Runtime/MovieSceneTracks/Private/Evaluation/MovieSceneSkeletalAnimationTemplate.cpp b/Engine/Source/Runtime/MovieSceneTracks/Private/Evaluation/MovieSceneSkeletalAnimationTemplate.cpp
--- a/Engine/Source/Runtime/MovieSceneTracks/Private/Evaluation/MovieSceneSkeletalAnimationTemplate.cpp
+++ b/Engine/Source/Runtime/MovieSceneTracks/Private/Evaluation/MovieSceneSkeletalAnimationTemplate.cpp
@@ -507,8 +507,7 @@ namespace MovieScene
 				const bool bContainsTransform = ContainsTransform(Player, SkeletalMeshComponent);
 				TOptional<FTransform> CurrentTransform = bContainsTransform ? SkeletalMeshComponent->GetRelativeTransform() : TOptional<FTransform>();
 								
-				const bool bSwapRootBoneWithComponentRoot = AnimSection->Params.bSwapRootBoneWithComponentRoot;
-				FAnimSequencerData AnimSequencerData(InAnimSequence, GetTypeHash(AnimTypeID), RootMotion, InFromPosition, InToPosition, Weight, bFireNotifies, bSwapRootBoneWithComponentRoot, CurrentTransform, AnimSection->Params.MirrorDataTable.Get());
+				FAnimSequencerData AnimSequencerData(InAnimSequence, GetTypeHash(AnimTypeID), RootMotion, InFromPosition, InToPosition, Weight, bFireNotifies, AnimSection->Params.SwapRootBone, CurrentTransform, AnimSection->Params.MirrorDataTable.Get());
 				SequencerInst->UpdateAnimTrackWithRootMotion(AnimSequencerData);
 			}
 			else if (UAnimInstance* AnimInst = GetSourceAnimInstance(SkeletalMeshComponent))
@@ -575,8 +574,7 @@ namespace MovieScene
 				const bool bContainsTransform = ContainsTransform(Player, SkeletalMeshComponent);
 				TOptional<FTransform> CurrentTransform = bContainsTransform ? SkeletalMeshComponent->GetRelativeTransform() : TOptional<FTransform>();
 
-				const bool bSwapRootBoneWithComponentRoot = AnimSection->Params.bSwapRootBoneWithComponentRoot;
-				FAnimSequencerData AnimSequencerData(InAnimSequence, GetTypeHash(AnimTypeID), RootMotion, InFromPosition, InToPosition, Weight, bFireNotifies, bSwapRootBoneWithComponentRoot, CurrentTransform, AnimSection->Params.MirrorDataTable.Get());
+				FAnimSequencerData AnimSequencerData(InAnimSequence, GetTypeHash(AnimTypeID), RootMotion, InFromPosition, InToPosition, Weight, bFireNotifies, AnimSection->Params.SwapRootBone, CurrentTransform, AnimSection->Params.MirrorDataTable.Get());
 				SequencerInst->UpdateAnimTrackWithRootMotion(AnimSequencerData);
 			}
 			else if (UAnimInstance* AnimInst = GetSourceAnimInstance(SkeletalMeshComponent))
diff --git a/Engine/Source/Runtime/MovieSceneTracks/Private/Sections/MovieSceneSkeletalAnimationSection.cpp b/Engine/Source/Runtime/MovieSceneTracks/Private/Sections/MovieSceneSkeletalAnimationSection.cpp
--- a/Engine/Source/Runtime/MovieSceneTracks/Private/Sections/MovieSceneSkeletalAnimationSection.cpp
+++ b/Engine/Source/Runtime/MovieSceneTracks/Private/Sections/MovieSceneSkeletalAnimationSection.cpp
@@ -3,6 +3,7 @@
 #include "Sections/MovieSceneSkeletalAnimationSection.h"
 #include "Channels/MovieSceneChannelProxy.h"
 #include "Animation/AnimSequence.h"
+#include "Runtime/AnimGraphRuntime/Public/AnimSequencerInstanceProxy.h"
 #include "Evaluation/MovieSceneSkeletalAnimationTemplate.h"
 #include "Logging/MessageLog.h"
 #include "MovieScene.h"
@@ -34,7 +35,7 @@ FMovieSceneSkeletalAnimationParams::FMovieSceneSkeletalAnimationParams()
 	Weight.SetDefault(1.f);
 	bSkipAnimNotifiers = false;
 	bForceCustomMode = false;
-	bSwapRootBoneWithComponentRoot = false;
+	SwapRootBone = ESwapRootBone::SwapRootBone_None;
 }
 
 UMovieSceneSkeletalAnimationSection::UMovieSceneSkeletalAnimationSection( const FObjectInitializer& ObjectInitializer )
diff --git a/Engine/Source/Runtime/MovieSceneTracks/Public/Sections/MovieSceneSkeletalAnimationSection.h b/Engine/Source/Runtime/MovieSceneTracks/Public/Sections/MovieSceneSkeletalAnimationSection.h
--- a/Engine/Source/Runtime/MovieSceneTracks/Public/Sections/MovieSceneSkeletalAnimationSection.h
+++ b/Engine/Source/Runtime/MovieSceneTracks/Public/Sections/MovieSceneSkeletalAnimationSection.h
@@ -12,6 +12,7 @@
 struct FMovieSceneSkeletalAnimRootMotionTrackParams;
 struct FAnimationPoseData;
 class UMirrorDataTable;
+enum class ESwapRootBone : uint8;
 
 USTRUCT(BlueprintType)
 struct FMovieSceneSkeletalAnimationParams
@@ -69,9 +70,9 @@ struct FMovieSceneSkeletalAnimationParams
 	UPROPERTY(EditAnywhere, Category = "Animation")
 	bool bForceCustomMode;
 
-	/** If on the root bone transform will be swapped to the component transform*/
+	/** If on the root bone transform will be swapped to the specified root*/
 	UPROPERTY(EditAnywhere, Category = "Animation")
-	bool bSwapRootBoneWithComponentRoot;
+	ESwapRootBone SwapRootBone;
 
 	UPROPERTY()
 	float StartOffset_DEPRECATED;
