Add GetRelevantAnimTimeRemaining & GetRelevantAnimTimeRemainingFraction to FAnimNode_StateMachine API
Add GetRelevantAnimTimeRemaining & GetRelevantAnimTimeRemainingFraction to AnimationStateMachineLibrary

diff --git a/Engine/Source/Runtime/AnimGraphRuntime/Private/AnimationStateMachineLibrary.cpp b/Engine/Source/Runtime/AnimGraphRuntime/Private/AnimationStateMachineLibrary.cpp
--- a/Engine/Source/Runtime/AnimGraphRuntime/Private/AnimationStateMachineLibrary.cpp
+++ b/Engine/Source/Runtime/AnimGraphRuntime/Private/AnimationStateMachineLibrary.cpp
@@ -129,3 +129,57 @@ FName UAnimationStateMachineLibrary::GetState(const FAnimUpdateContext& UpdateCo
 
 	return OutName;
 }
+
+float UAnimationStateMachineLibrary::GetRelevantAnimTimeRemaining(const FAnimUpdateContext& UpdateContext, const FAnimationStateResultReference& Node)
+{
+	float Result = MAX_flt;
+
+	Node.CallAnimNodeFunction<FAnimNode_StateResult>(
+		TEXT("IsStateBlendingOut"),
+		[&UpdateContext, &Result](FAnimNode_StateResult& StateResultNode)
+	{
+		if (const FAnimationUpdateContext* AnimationUpdateContext = UpdateContext.GetContext())
+		{
+			IAnimClassInterface* AnimBlueprintClass = AnimationUpdateContext->GetAnimClass();
+
+			// Previous node to an FAnimNode_StateResult is always its owning FAnimNode_StateMachine
+			const int32 MachineIndex = AnimBlueprintClass->GetAnimNodeProperties().Num() - 1 - AnimationUpdateContext->GetPreviousNodeId();
+			const int32 StateIndex = StateResultNode.GetStateIndex();
+
+			const FAnimInstanceProxy* AnimInstanceProxy = AnimationUpdateContext->AnimInstanceProxy;
+			if (const FAnimNode_StateMachine* MachineInstance = AnimInstanceProxy->GetStateMachineInstance(MachineIndex))
+			{
+				Result = MachineInstance->GetRelevantAnimTimeRemaining(AnimInstanceProxy, StateIndex);
+			}
+		}
+	});
+
+	return Result;
+}
+
+float UAnimationStateMachineLibrary::GetRelevantAnimTimeRemainingFraction(const FAnimUpdateContext& UpdateContext, const FAnimationStateResultReference& Node)
+{
+	float Result = 1.f;
+
+	Node.CallAnimNodeFunction<FAnimNode_StateResult>(
+		TEXT("IsStateBlendingOut"),
+		[&UpdateContext, &Result](FAnimNode_StateResult& StateResultNode)
+	{
+		if (const FAnimationUpdateContext* AnimationUpdateContext = UpdateContext.GetContext())
+		{
+			IAnimClassInterface* AnimBlueprintClass = AnimationUpdateContext->GetAnimClass();
+
+			// Previous node to an FAnimNode_StateResult is always its owning FAnimNode_StateMachine
+			const int32 MachineIndex = AnimBlueprintClass->GetAnimNodeProperties().Num() - 1 - AnimationUpdateContext->GetPreviousNodeId();
+			const int32 StateIndex = StateResultNode.GetStateIndex();
+
+			const FAnimInstanceProxy* AnimInstanceProxy = AnimationUpdateContext->AnimInstanceProxy;
+			if (const FAnimNode_StateMachine* MachineInstance = AnimInstanceProxy->GetStateMachineInstance(MachineIndex))
+			{
+				Result = MachineInstance->GetRelevantAnimTimeRemainingFraction(AnimInstanceProxy, StateIndex);
+			}
+		}
+	});
+
+	return Result;
+}
\ No newline at end of file
diff --git a/Engine/Source/Runtime/AnimGraphRuntime/Public/AnimationStateMachineLibrary.h b/Engine/Source/Runtime/AnimGraphRuntime/Public/AnimationStateMachineLibrary.h
--- a/Engine/Source/Runtime/AnimGraphRuntime/Public/AnimationStateMachineLibrary.h
+++ b/Engine/Source/Runtime/AnimGraphRuntime/Public/AnimationStateMachineLibrary.h
@@ -80,4 +80,12 @@ public:
 	/** Returns the name of the current state of this state machine */
 	UFUNCTION(BlueprintPure, Category = "State Machine", meta=(BlueprintThreadSafe))
 	static FName GetState(const FAnimUpdateContext& UpdateContext, const FAnimationStateMachineReference& Node);
+
+	/** Returns the remaining animation time of the state's most relevant asset player */
+	UFUNCTION(BlueprintPure, Category = "State Machine", meta = (BlueprintThreadSafe))
+	static float GetRelevantAnimTimeRemaining(const FAnimUpdateContext& UpdateContext, const FAnimationStateResultReference& Node);
+
+	/** Returns the remaining animation time as a fraction of the duration for the state's most relevant asset player */
+	UFUNCTION(BlueprintPure, Category = "State Machine", meta = (BlueprintThreadSafe))
+	static float GetRelevantAnimTimeRemainingFraction(const FAnimUpdateContext& UpdateContext, const FAnimationStateResultReference& Node);
 };
\ No newline at end of file
diff --git a/Engine/Source/Runtime/Engine/Classes/Animation/AnimNode_StateMachine.h b/Engine/Source/Runtime/Engine/Classes/Animation/AnimNode_StateMachine.h
--- a/Engine/Source/Runtime/Engine/Classes/Animation/AnimNode_StateMachine.h
+++ b/Engine/Source/Runtime/Engine/Classes/Animation/AnimNode_StateMachine.h
@@ -292,6 +292,20 @@ protected:
 	void EvaluateTransitionStandardBlendInternal(FPoseContext& Output, FAnimationActiveTransitionEntry& Transition, const FPoseContext& PreviousStateResult, const FPoseContext& NextStateResult);
 	void EvaluateTransitionCustomBlend(FPoseContext& Output, FAnimationActiveTransitionEntry& Transition, bool bIntermediatePoseIsValid);
 
+	// Get the time remaining in seconds for the most relevant animation in the source state 
+	float GetRelevantAnimTimeRemaining(const FAnimInstanceProxy* InAnimInstanceProxy, int32 StateIndex) const;
+	float GetRelevantAnimTimeRemaining(const FAnimationUpdateContext& Context, int32 StateIndex) const
+	{
+		return GetRelevantAnimTimeRemaining(Context.AnimInstanceProxy, StateIndex);
+	}
+
+	// Get the time remaining as a fraction of the duration for the most relevant animation in the source state 
+	float GetRelevantAnimTimeRemainingFraction(const FAnimInstanceProxy* InAnimInstanceProxy, int32 StateIndex) const;
+	float GetRelevantAnimTimeRemainingFraction(const FAnimationUpdateContext& Context, int32 StateIndex) const
+	{
+		return GetRelevantAnimTimeRemainingFraction(Context.AnimInstanceProxy, StateIndex);
+	}
+
 	const FAnimNode_AssetPlayerBase* GetRelevantAssetPlayerFromState(const FAnimInstanceProxy* InAnimInstanceProxy, const FBakedAnimationState& StateInfo) const;
 	const FAnimNode_AssetPlayerBase* GetRelevantAssetPlayerFromState(const FAnimationUpdateContext& Context, const FBakedAnimationState& StateInfo) const
 	{
@@ -320,4 +334,5 @@ protected:
 
 public:
 	friend struct FAnimInstanceProxy;
+	friend class UAnimationStateMachineLibrary;
 };
diff --git a/Engine/Source/Runtime/Engine/Private/Animation/AnimInstanceProxy.cpp b/Engine/Source/Runtime/Engine/Private/Animation/AnimInstanceProxy.cpp
--- a/Engine/Source/Runtime/Engine/Private/Animation/AnimInstanceProxy.cpp
+++ b/Engine/Source/Runtime/Engine/Private/Animation/AnimInstanceProxy.cpp
@@ -2335,12 +2335,9 @@ float FAnimInstanceProxy::GetInstanceTransitionTimeElapsedFraction(int32 Machine
 
 float FAnimInstanceProxy::GetRelevantAnimTimeRemaining(int32 MachineIndex, int32 StateIndex) const
 {
-	if(const FAnimNode_AssetPlayerBase* AssetPlayer = GetRelevantAssetPlayerFromState(MachineIndex, StateIndex))
+	if (const FAnimNode_StateMachine* StateMachine = GetStateMachineInstance(MachineIndex))
 	{
-		if(AssetPlayer->GetAnimAsset())
-		{
-			return AssetPlayer->GetCurrentAssetLength() - AssetPlayer->GetCurrentAssetTimePlayRateAdjusted();
-		}
+		return StateMachine->GetRelevantAnimTimeRemaining(this, StateIndex);
 	}
 
 	return MAX_flt;
@@ -2348,16 +2345,9 @@ float FAnimInstanceProxy::GetRelevantAnimTimeRemaining(int32 MachineIndex, int32
 
 float FAnimInstanceProxy::GetRelevantAnimTimeRemainingFraction(int32 MachineIndex, int32 StateIndex) const
 {
-	if(const FAnimNode_AssetPlayerBase* AssetPlayer = GetRelevantAssetPlayerFromState(MachineIndex, StateIndex))
+	if (const FAnimNode_StateMachine* StateMachine = GetStateMachineInstance(MachineIndex))
 	{
-		if(AssetPlayer->GetAnimAsset())
-		{
-			float Length = AssetPlayer->GetCurrentAssetLength();
-			if(Length > 0.0f)
-			{
-				return (Length - AssetPlayer->GetCurrentAssetTimePlayRateAdjusted()) / Length;
-			}
-		}
+		return StateMachine->GetRelevantAnimTimeRemainingFraction(this, StateIndex);
 	}
 
 	return 1.0f;
diff --git a/Engine/Source/Runtime/Engine/Private/Animation/AnimNode_StateMachine.cpp b/Engine/Source/Runtime/Engine/Private/Animation/AnimNode_StateMachine.cpp
--- a/Engine/Source/Runtime/Engine/Private/Animation/AnimNode_StateMachine.cpp
+++ b/Engine/Source/Runtime/Engine/Private/Animation/AnimNode_StateMachine.cpp
@@ -553,6 +553,36 @@ void FAnimNode_StateMachine::Update_AnyThread(const FAnimationUpdateContext& Con
 #endif //ANIM_TRACE_ENABLED
 }
 
+float FAnimNode_StateMachine::GetRelevantAnimTimeRemaining(const FAnimInstanceProxy* InAnimInstanceProxy, int32 StateIndex) const
+{
+	if (const FAnimNode_AssetPlayerBase* AssetPlayer = GetRelevantAssetPlayerFromState(InAnimInstanceProxy, GetStateInfo(StateIndex)))
+	{
+		if (AssetPlayer->GetAnimAsset())
+		{
+			return AssetPlayer->GetCurrentAssetLength() - AssetPlayer->GetCurrentAssetTimePlayRateAdjusted();
+		}
+	}
+
+	return MAX_flt;
+}
+
+float FAnimNode_StateMachine::GetRelevantAnimTimeRemainingFraction(const FAnimInstanceProxy* InAnimInstanceProxy, int32 StateIndex) const
+{
+	if (const FAnimNode_AssetPlayerBase* AssetPlayer = GetRelevantAssetPlayerFromState(InAnimInstanceProxy, GetStateInfo(StateIndex)))
+	{
+		if (AssetPlayer->GetAnimAsset())
+		{
+			float Length = AssetPlayer->GetCurrentAssetLength();
+			if (Length > 0.0f)
+			{
+				return (Length - AssetPlayer->GetCurrentAssetTimePlayRateAdjusted()) / Length;
+			}
+		}
+	}
+
+	return 1.f;
+}
+
 const FAnimNode_AssetPlayerBase* FAnimNode_StateMachine::GetRelevantAssetPlayerFromState(const FAnimInstanceProxy* InAnimInstanceProxy, const FBakedAnimationState& StateInfo) const
 {
 	const FAnimNode_AssetPlayerBase* ResultPlayer = nullptr;
