Fixed various issues with the revised re-instancing code path for anim BPs

Fixed RBAN not re-initialzing correctly on re-instancing - added a special case to reinit the simulation.
Made reinstancing handler virtual so users can add their own fallbacks in the future.
Fixed dynamic AnimGraphLayers not (un)linking correctly after live-recompile, as it was trying to link against the old (replaced) animation instance

diff --git a/Engine/Source/Runtime/AnimGraphRuntime/Private/BoneControllers/AnimNode_RigidBody.cpp b/Engine/Source/Runtime/AnimGraphRuntime/Private/BoneControllers/AnimNode_RigidBody.cpp
--- a/Engine/Source/Runtime/AnimGraphRuntime/Private/BoneControllers/AnimNode_RigidBody.cpp
+++ b/Engine/Source/Runtime/AnimGraphRuntime/Private/BoneControllers/AnimNode_RigidBody.cpp
@@ -224,6 +224,18 @@ void FAnimNode_RigidBody::GatherDebugData(FNodeDebugData& DebugData)
 	}
 }
 
+void FAnimNode_RigidBody::Initialize_AnyThread(const FAnimationInitializeContext& Context)
+{
+	Super::Initialize_AnyThread(Context);
+
+#if WITH_EDITOR
+	if(GIsReinstancing)
+	{
+		InitPhysics(Cast<UAnimInstance>(Context.GetAnimInstanceObject()));
+	}
+#endif
+}
+
 FTransform SpaceToWorldTransform(ESimulationSpace Space, const FTransform& ComponentToWorld, const FTransform& BaseBoneTM)
 {
 	switch (Space)
diff --git a/Engine/Source/Runtime/AnimGraphRuntime/Public/BoneControllers/AnimNode_RigidBody.h b/Engine/Source/Runtime/AnimGraphRuntime/Public/BoneControllers/AnimNode_RigidBody.h
--- a/Engine/Source/Runtime/AnimGraphRuntime/Public/BoneControllers/AnimNode_RigidBody.h
+++ b/Engine/Source/Runtime/AnimGraphRuntime/Public/BoneControllers/AnimNode_RigidBody.h
@@ -137,6 +137,7 @@ struct ANIMGRAPHRUNTIME_API FAnimNode_RigidBody : public FAnimNode_SkeletalContr
 
 	// FAnimNode_Base interface
 	virtual void GatherDebugData(FNodeDebugData& DebugData) override;
+	virtual void Initialize_AnyThread(const FAnimationInitializeContext& Context) override;
 	// End of FAnimNode_Base interface
 
 	// FAnimNode_SkeletalControlBase interface
diff --git a/Engine/Source/Runtime/Engine/Classes/Animation/AnimInstance.h b/Engine/Source/Runtime/Engine/Classes/Animation/AnimInstance.h
--- a/Engine/Source/Runtime/Engine/Classes/Animation/AnimInstance.h
+++ b/Engine/Source/Runtime/Engine/Classes/Animation/AnimInstance.h
@@ -1637,7 +1637,7 @@ public:
 	void DispatchQueuedAnimEvents();
 
 #if WITH_EDITOR
-	// Helper function to handle reinstancing
-	void HandleAnimInstanceReplaced();
+	// Helper function to handle reinstancing in editor
+	virtual void HandleAnimInstanceReplaced();
 #endif
 };
diff --git a/Engine/Source/Runtime/Engine/Private/Animation/AnimNode_LinkedAnimGraph.cpp b/Engine/Source/Runtime/Engine/Private/Animation/AnimNode_LinkedAnimGraph.cpp
--- a/Engine/Source/Runtime/Engine/Private/Animation/AnimNode_LinkedAnimGraph.cpp
+++ b/Engine/Source/Runtime/Engine/Private/Animation/AnimNode_LinkedAnimGraph.cpp
@@ -484,6 +484,7 @@ void FAnimNode_LinkedAnimGraph::HandleObjectsReplaced(const TMap<UObject*, UObje
 			if(UAnimInstance* ReinstancedTargetInstance = Cast<UAnimInstance>(*ReinstancedTarget))
 			{
 				DynamicUnlink(ThisTargetInstance);
+				SetTargetInstance(ReinstancedTargetInstance);
 				DynamicLink(ReinstancedTargetInstance);
 			}
 		}
