Physics interface cleanup.

diff --git a/Engine/Source/Editor/ClothingSystemEditor/Public/ClothingAssetFactory.h b/Engine/Source/Editor/ClothingSystemEditor/Public/ClothingAssetFactory.h
--- a/Engine/Source/Editor/ClothingSystemEditor/Public/ClothingAssetFactory.h
+++ b/Engine/Source/Editor/ClothingSystemEditor/Public/ClothingAssetFactory.h
@@ -50,27 +50,6 @@ public:
 
 private:
 
-#if WITH_APEX_CLOTHING
-
-	struct FApexVertData
-	{
-		uint16 BoneIndices[MAX_TOTAL_INFLUENCES];
-	};
-
-	// Convert from APEX to UE coodinate system
-	nvidia::apex::ClothingAsset* ConvertApexAssetCoordSystem(nvidia::apex::ClothingAsset* InAsset);
-
-	// Convert APEX UV direction to UE UV direction
-	void FlipAuthoringUvs(NvParameterized::Interface* InRenderMeshAuthoringInterface, bool bFlipU, bool bFlipV);
-
-	// Extraction methods for pulling the required data from an APEX asset and
-	// pushing it to a UClothingAssetCommon
-	void ExtractLodPhysicalData(UClothingAssetCommon* NewAsset, nvidia::apex::ClothingAsset& InApexAsset, int32 InLodIdx, FClothLODDataCommon& InLodData, TArray<FApexVertData>& OutApexVertData);
-	void ExtractBoneData(UClothingAssetCommon* NewAsset, nvidia::apex::ClothingAsset& InApexAsset);
-	void ExtractSphereCollisions(UClothingAssetCommon* NewAsset, nvidia::apex::ClothingAsset& InApexAsset, int32 InLodIdx, FClothLODDataCommon& InLodData);
-	void ExtractMaterialParameters(UClothingAssetCommon* NewAsset, nvidia::apex::ClothingAsset& InApexAsset);
-#endif
-
 	// Utility methods for skeletal mesh extraction //////////////////////////
 
 	/** Handles internal import of LODs */
diff --git a/Engine/Source/Editor/ClothingSystemEditor/Public/ClothingSystemEditorModule.h b/Engine/Source/Editor/ClothingSystemEditor/Public/ClothingSystemEditorModule.h
--- a/Engine/Source/Editor/ClothingSystemEditor/Public/ClothingSystemEditorModule.h
+++ b/Engine/Source/Editor/ClothingSystemEditor/Public/ClothingSystemEditorModule.h
@@ -22,8 +22,4 @@ public:
 
 private:
 
-#if WITH_NVCLOTH
-	FSimulationEditorExtenderNv NvEditorExtender;
-#endif
-
 };
diff --git a/Engine/Source/Editor/ClothingSystemEditor/Public/SimulationEditorExtenderNv.h b/Engine/Source/Editor/ClothingSystemEditor/Public/SimulationEditorExtenderNv.h
--- a/Engine/Source/Editor/ClothingSystemEditor/Public/SimulationEditorExtenderNv.h
+++ b/Engine/Source/Editor/ClothingSystemEditor/Public/SimulationEditorExtenderNv.h
@@ -1,89 +1,3 @@
 // Copyright Epic Games, Inc. All Rights Reserved.
 
 #pragma once
-
-#if WITH_NVCLOTH
-
-#include "SimulationEditorExtender.h"
-#include "Internationalization/Text.h"
-
-struct FNvVisualizationOptions
-{
-	FNvVisualizationOptions();
-
-	// FOptionData contains a localised name and other state for visualization menu options
-	struct FOptionData
-	{
-		FOptionData()
-			: bDisablesSimulation(false)
-		{}
-
-		// Text for menu entries
-		FText DisplayName;
-		FText ToolTip;
-
-		// Whether or not this option requires the simulation to be disabled
-		bool bDisablesSimulation;
-	};
-
-	// List of options
-	enum class EVisualizationOption : uint8
-	{
-		PhysMesh = 0,
-		Normals,
-		Collision,
-		Backstop,
-		MaxDistances,
-		SelfCollision,
-		AnimDrive,
-		Max
-	};
-
-	// Actual option entries
-	FOptionData OptionData[(int32)EVisualizationOption::Max];
-
-	// Flags determining which options are enabled
-	bool Flags[(int32)EVisualizationOption::Max];
-
-	// Is an option set
-	bool IsSet(const EVisualizationOption InOption) const
-	{
-		return Flags[(int32)InOption];
-	}
-
-	// Toggle an option
-	void Toggle(const EVisualizationOption InOption)
-	{
-		Flags[(int32)InOption] = !Flags[(int32)InOption];
-	}
-
-	// Whether or not - given the current enabled options - the simulation should be disabled
-	bool ShouldDisableSimulation() const;
-};
-
-class FSimulationEditorExtenderNv : public ISimulationEditorExtender
-{
-
-public:
-
-	// ISimulationEditorExtender Interface
-	virtual UClass* GetSupportedSimulationFactoryClass() override;
-	virtual void ExtendViewportShowMenu(FMenuBuilder& InMenuBuilder, TSharedRef<IPersonaPreviewScene> InPreviewScene) override;
-	virtual void DebugDrawSimulation(const IClothingSimulation* InSimulation, USkeletalMeshComponent* InOwnerComponent, FPrimitiveDrawInterface* PDI) override;
-	virtual void DebugDrawSimulationTexts(const IClothingSimulation* InSimulation, USkeletalMeshComponent* InOwnerComponent, FCanvas* Canvas, const FSceneView* SceneView) override {}
-	// End ISimulationEditorExtender Interface
-
-private:
-
-	// Handler for visualization entry being clicked
-	void OnEntryClicked(FNvVisualizationOptions::EVisualizationOption InOption, TSharedRef<IPersonaPreviewScene> InPreviewScene);
-
-	// Checkstate function for visualization entries
-	bool IsEntryChecked(FNvVisualizationOptions::EVisualizationOption InOption) const;
-
-	// Visualization options for NvCloth
-	FNvVisualizationOptions VisulizationFlags;
-
-};
-
-#endif
diff --git a/Engine/Source/Editor/DetailCustomizations/Private/BodyInstanceCustomization.cpp b/Engine/Source/Editor/DetailCustomizations/Private/BodyInstanceCustomization.cpp
--- a/Engine/Source/Editor/DetailCustomizations/Private/BodyInstanceCustomization.cpp
+++ b/Engine/Source/Editor/DetailCustomizations/Private/BodyInstanceCustomization.cpp
@@ -1294,8 +1294,7 @@ void FBodyInstanceCustomizationHelper::CustomizeDetails( IDetailLayoutBuilder& D
 
 		AddMaxAngularVelocity(PhysicsCategory, BodyInstanceHandler);
 
-#if WITH_CHAOS
-		// Hide PhysX-Only settings in Chaos
+		// Hide legacy settings in Chaos
 		BodyInstanceHandler->GetChildHandle(GET_MEMBER_NAME_CHECKED(FBodyInstance, DOFMode))->MarkHiddenByCustomization();
 		BodyInstanceHandler->GetChildHandle(GET_MEMBER_NAME_CHECKED(FBodyInstance, bLockTranslation))->MarkHiddenByCustomization();
 		BodyInstanceHandler->GetChildHandle(GET_MEMBER_NAME_CHECKED(FBodyInstance, bLockRotation))->MarkHiddenByCustomization();
@@ -1309,7 +1308,7 @@ void FBodyInstanceCustomizationHelper::CustomizeDetails( IDetailLayoutBuilder& D
 		BodyInstanceHandler->GetChildHandle(GET_MEMBER_NAME_CHECKED(FBodyInstance, VelocitySolverIterationCount))->MarkHiddenByCustomization();
 		BodyInstanceHandler->GetChildHandle(GET_MEMBER_NAME_CHECKED(FBodyInstance, MaxDepenetrationVelocity))->MarkHiddenByCustomization();
 		BodyInstanceHandler->GetChildHandle(GET_MEMBER_NAME_CHECKED(FBodyInstance, CustomDOFPlaneNormal))->MarkHiddenByCustomization();
-#endif
+
 		//Add the rest
 		uint32 NumChildren = 0;
 		BodyInstanceHandler->GetNumChildren(NumChildren);
diff --git a/Engine/Source/Editor/DetailCustomizations/Private/PhysicsConstraintComponentDetails.cpp b/Engine/Source/Editor/DetailCustomizations/Private/PhysicsConstraintComponentDetails.cpp
--- a/Engine/Source/Editor/DetailCustomizations/Private/PhysicsConstraintComponentDetails.cpp
+++ b/Engine/Source/Editor/DetailCustomizations/Private/PhysicsConstraintComponentDetails.cpp
@@ -320,9 +320,6 @@ void FPhysicsConstraintComponentDetails::AddLinearLimits(IDetailLayoutBuilder& D
 	LinearLimitCat.AddProperty(LinearConstraintProperty->GetChildHandle(GET_MEMBER_NAME_CHECKED(FLinearConstraint, Stiffness))).IsEnabled(TAttribute<bool>::Create(TAttribute<bool>::FGetter::CreateLambda(IsLinearMotionLimited)));
 	LinearLimitCat.AddProperty(LinearConstraintProperty->GetChildHandle(GET_MEMBER_NAME_CHECKED(FLinearConstraint, Damping))).IsEnabled(TAttribute<bool>::Create(TAttribute<bool>::FGetter::CreateLambda(IsLinearMotionLimited)));
 	LinearLimitCat.AddProperty(LinearConstraintProperty->GetChildHandle(GET_MEMBER_NAME_CHECKED(FLinearConstraint, Restitution))).IsEnabled(TAttribute<bool>::Create(TAttribute<bool>::FGetter::CreateLambda(IsRestitutionEnabled)));
-#if !WITH_CHAOS
-	LinearLimitCat.AddProperty(LinearConstraintProperty->GetChildHandle(GET_MEMBER_NAME_CHECKED(FLinearConstraint, ContactDistance))).IsEnabled(TAttribute<bool>::Create(TAttribute<bool>::FGetter::CreateLambda(IsLinearMotionLimited)));
-#endif
 	LinearLimitCat.AddProperty(ProfilePropertiesProperty->GetChildHandle(GET_MEMBER_NAME_CHECKED(FConstraintProfileProperties, bLinearBreakable)));
 	LinearLimitCat.AddProperty(ProfilePropertiesProperty->GetChildHandle(GET_MEMBER_NAME_CHECKED(FConstraintProfileProperties, LinearBreakThreshold)));
 
@@ -452,9 +449,6 @@ void FPhysicsConstraintComponentDetails::AddAngularLimits(IDetailLayoutBuilder&
 	SwingGroup.AddPropertyRow(ConeConstraintProperty->GetChildHandle(GET_MEMBER_NAME_CHECKED(FConeConstraint, Stiffness)).ToSharedRef()).IsEnabled(TAttribute<bool>::Create(TAttribute<bool>::FGetter::CreateSP(this, &FPhysicsConstraintComponentDetails::IsPropertyEnabled, EPropertyType::AngularSwingLimit)));
 	SwingGroup.AddPropertyRow(ConeConstraintProperty->GetChildHandle(GET_MEMBER_NAME_CHECKED(FConeConstraint, Damping)).ToSharedRef()).IsEnabled(TAttribute<bool>::Create(TAttribute<bool>::FGetter::CreateSP(this, &FPhysicsConstraintComponentDetails::IsPropertyEnabled, EPropertyType::AngularSwingLimit)));
 	SwingGroup.AddPropertyRow(ConeConstraintProperty->GetChildHandle(GET_MEMBER_NAME_CHECKED(FConeConstraint, Restitution)).ToSharedRef()).IsEnabled(TAttribute<bool>::Create(TAttribute<bool>::FGetter::CreateLambda(SwingRestitutionEnabled)));
-#if !WITH_CHAOS
-	SwingGroup.AddPropertyRow(ConeConstraintProperty->GetChildHandle(GET_MEMBER_NAME_CHECKED(FConeConstraint, ContactDistance)).ToSharedRef()).IsEnabled(TAttribute<bool>::Create(TAttribute<bool>::FGetter::CreateSP(this, &FPhysicsConstraintComponentDetails::IsPropertyEnabled, EPropertyType::AngularSwingLimit)));
-#endif
 
 	TSharedPtr<IPropertyHandle> SoftTwistProperty = TwistConstraintProperty->GetChildHandle(GET_MEMBER_NAME_CHECKED(FTwistConstraint, bSoftConstraint));
 	auto TwistRestitutionEnabled = [this, SoftTwistProperty]()
@@ -468,9 +462,6 @@ void FPhysicsConstraintComponentDetails::AddAngularLimits(IDetailLayoutBuilder&
 	TwistGroup.AddPropertyRow(TwistConstraintProperty->GetChildHandle(GET_MEMBER_NAME_CHECKED(FTwistConstraint, Stiffness)).ToSharedRef()).IsEnabled(TAttribute<bool>::Create(TAttribute<bool>::FGetter::CreateSP(this, &FPhysicsConstraintComponentDetails::IsPropertyEnabled, EPropertyType::AngularTwistLimit)));
 	TwistGroup.AddPropertyRow(TwistConstraintProperty->GetChildHandle(GET_MEMBER_NAME_CHECKED(FTwistConstraint, Damping)).ToSharedRef()).IsEnabled(TAttribute<bool>::Create(TAttribute<bool>::FGetter::CreateSP(this, &FPhysicsConstraintComponentDetails::IsPropertyEnabled, EPropertyType::AngularTwistLimit)));
 	TwistGroup.AddPropertyRow(TwistConstraintProperty->GetChildHandle(GET_MEMBER_NAME_CHECKED(FTwistConstraint, Restitution)).ToSharedRef()).IsEnabled(TAttribute<bool>::Create(TAttribute<bool>::FGetter::CreateLambda(TwistRestitutionEnabled)));
-#if !WITH_CHAOS
-	TwistGroup.AddPropertyRow(TwistConstraintProperty->GetChildHandle(GET_MEMBER_NAME_CHECKED(FTwistConstraint, ContactDistance)).ToSharedRef()).IsEnabled(TAttribute<bool>::Create(TAttribute<bool>::FGetter::CreateSP(this, &FPhysicsConstraintComponentDetails::IsPropertyEnabled, EPropertyType::AngularTwistLimit)));
-#endif
 
 	if (bInPhat == false)
 	{
diff --git a/Engine/Source/Editor/Persona/Private/PersonaMeshDetails.cpp b/Engine/Source/Editor/Persona/Private/PersonaMeshDetails.cpp
--- a/Engine/Source/Editor/Persona/Private/PersonaMeshDetails.cpp
+++ b/Engine/Source/Editor/Persona/Private/PersonaMeshDetails.cpp
@@ -45,12 +45,6 @@
 #include "IContentBrowserSingleton.h"
 #include "ContentBrowserModule.h"
 #include "EditorFramework/AssetImportData.h"
-
-#if WITH_APEX_CLOTHING
-	#include "ApexClothingUtils.h"
-	#include "ApexClothingOptionWindow.h"
-#endif // #if WITH_APEX_CLOTHING
-
 #include "ClothingAsset.h"
 
 #include "LODUtilities.h"
@@ -2582,10 +2576,7 @@ void FPersonaMeshDetails::AddLODLevelCategories(IDetailLayoutBuilder& DetailLayo
 	{
 		const int32 SkelMeshLODCount = SkelMesh->GetLODNum();
 
-
-#if WITH_APEX_CLOTHING || WITH_CHAOS_CLOTHING
 		ClothComboBoxes.Reset();
-#endif
 
 		//Create material list panel to let users control the materials array
 		{
@@ -5010,8 +5001,6 @@ TSharedRef<SWidget> FPersonaMeshDetails::OnGenerateCustomSectionWidgetsForSectio
 		}
 	}
 
-#if WITH_APEX_CLOTHING || WITH_CHAOS_CLOTHING
-
 	UpdateClothingEntries();
 
 	ClothComboBoxes.AddDefaulted();
@@ -5051,7 +5040,6 @@ TSharedRef<SWidget> FPersonaMeshDetails::OnGenerateCustomSectionWidgetsForSectio
 			]
 		]
 	];
-#endif// #if WITH_APEX_CLOTHING || WITH_CHAOS_CLOTHING
 
 	SectionWidget->AddSlot()
 	.AutoHeight()
@@ -6048,31 +6036,6 @@ void FPersonaMeshDetails::CustomizeClothingProperties(IDetailLayoutBuilder& Deta
 
 		ClothingFilesCategory.AddCustomBuilder(ClothingAssetsPropertyBuilder, false);
 	}
-
-#if WITH_APEX_CLOTHING
-	// Button to add a new clothing file
-	ClothingFilesCategory.AddCustomRow( LOCTEXT("AddAPEXClothingFileFilterString", "Add APEX clothing file"))
-	[
-		SNew(SHorizontalBox)
-		 
-		+SHorizontalBox::Slot()
-		.AutoWidth()
-		[
-			SNew(SButton)
-			.OnClicked(this, &FPersonaMeshDetails::OnOpenClothingFileClicked, &DetailLayout)
-			.ToolTip(IDocumentation::Get()->CreateToolTip(
-				LOCTEXT("AddClothingButtonTooltip", "Select a new APEX clothing file and add it to the skeletal mesh."),
-				NULL,
-				TEXT("Shared/Editors/Persona"),
-				TEXT("AddClothing")))
-			[
-				SNew(STextBlock)
-				.Font(IDetailLayoutBuilder::GetDetailFont())
-				.Text(LOCTEXT("AddAPEXClothingFile", "Add APEX clothing file..."))
-			]
-		]
-	];
-#endif
 }
 
 //
@@ -6088,30 +6051,6 @@ void FPersonaMeshDetails::OnGenerateElementForClothingAsset( TSharedRef<IPropert
 		SNew(SHorizontalBox)
 		+ SHorizontalBox::Slot()
 		.FillWidth(1)
-
-#if WITH_APEX_CLOTHING
-		// re-import button
-		+ SHorizontalBox::Slot()
-		.VAlign( VAlign_Center )
-		.Padding(2)
-		.AutoWidth()
-		[
-			SNew( SButton )
-			.Text( LOCTEXT("ReimportButtonLabel", "Reimport") )
-			.OnClicked(this, &FPersonaMeshDetails::OnReimportApexFileClicked, ElementIndex, DetailLayout)
-			.IsFocusable( false )
-			.ContentPadding(0)
-			.ForegroundColor( FSlateColor::UseForeground() )
-			.ButtonColorAndOpacity(FLinearColor(1.0f,1.0f,1.0f,0.0f))
-			.ToolTipText(LOCTEXT("ReimportApexFileTip", "Reimport this APEX asset"))
-			[ 
-				SNew( SImage )
-				.Image( FAppStyle::GetBrush("Persona.ReimportAsset") )
-				.ColorAndOpacity( FSlateColor::UseForeground() )
-			]
-		]
-#endif  // #if WITH_APEX_CLOTHING
-
 		// remove button
 		+ SHorizontalBox::Slot()
 		.VAlign( VAlign_Center )
@@ -6323,64 +6262,6 @@ TSharedRef<SUniformGridPanel> FPersonaMeshDetails::MakeClothingDetailsWidget(int
 	return Grid;
 }
 
-#if WITH_APEX_CLOTHING
-FReply FPersonaMeshDetails::OnReimportApexFileClicked(int32 AssetIndex, IDetailLayoutBuilder* DetailLayout)
-{
-	USkeletalMesh* SkelMesh = GetPersonaToolkit()->GetMesh();
-
-	check(SkelMesh && SkelMesh->GetMeshClothingAssets().IsValidIndex(AssetIndex));
-
-	UClothingAssetBase* AssetToReimport = SkelMesh->GetMeshClothingAssets()[AssetIndex];
-	check(AssetToReimport);
-
-	FString ReimportPath = AssetToReimport->ImportedFilePath;
-
-	if(ReimportPath.IsEmpty())
-	{
-		const FText MessageText = LOCTEXT("Warning_NoReimportPath", "There is no reimport path available for this asset, it was likely created in the Editor. Would you like to select a file and overwrite this asset?");
-		EAppReturnType::Type MessageReturn = FMessageDialog::Open(EAppMsgType::YesNo, MessageText);
-
-		if(MessageReturn == EAppReturnType::Yes)
-		{
-			ReimportPath = ApexClothingUtils::PromptForClothingFile();
-		}
-	}
-
-	if(ReimportPath.IsEmpty())
-	{
-		return FReply::Handled();
-	}
-
-	// Retry if the file isn't there
-	if(!FPaths::FileExists(ReimportPath))
-	{
-		const FText MessageText = LOCTEXT("Warning_NoFileFound", "Could not find an asset to reimport, select a new file on disk?");
-		EAppReturnType::Type MessageReturn = FMessageDialog::Open(EAppMsgType::YesNo, MessageText);
-
-		if(MessageReturn == EAppReturnType::Yes)
-		{
-			ReimportPath = ApexClothingUtils::PromptForClothingFile();
-		}
-	}
-
-	FClothingSystemEditorInterfaceModule& ClothingEditorInterface = FModuleManager::Get().LoadModuleChecked<FClothingSystemEditorInterfaceModule>("ClothingSystemEditorInterface");
-	UClothingAssetFactoryBase* Factory = ClothingEditorInterface.GetClothingAssetFactory();
-
-	if(Factory && Factory->CanImport(ReimportPath))
-	{
-		Factory->Reimport(ReimportPath, SkelMesh, AssetToReimport);
-
-		UpdateClothingEntries();
-		RefreshClothingComboBoxes();
-
-		// Force layout to refresh
-		RefreshMeshDetailLayout();
-	}
-
-	return FReply::Handled();
-}
-#endif
-
 FReply FPersonaMeshDetails::OnRemoveClothingAssetClicked(int32 AssetIndex, IDetailLayoutBuilder* DetailLayout)
 {
 	USkeletalMesh* SkelMesh = GetPersonaToolkit()->GetMesh();
@@ -6427,23 +6308,6 @@ FReply FPersonaMeshDetails::OnRemoveClothingAssetClicked(int32 AssetIndex, IDeta
 	return FReply::Handled();
 }
 
-#if WITH_APEX_CLOTHING
-FReply FPersonaMeshDetails::OnOpenClothingFileClicked(IDetailLayoutBuilder* DetailLayout)
-{
-	USkeletalMesh* SkelMesh = GetPersonaToolkit()->GetMesh();
-
-	if(SkelMesh)
-	{
-		ApexClothingUtils::PromptAndImportClothing(SkelMesh);
-		
-		UpdateClothingEntries();
-		RefreshClothingComboBoxes();
-	}
-
-	return FReply::Handled();
-}
-#endif
-
 void FPersonaMeshDetails::UpdateClothingEntries()
 {
 	USkeletalMesh* Mesh = GetPersonaToolkit()->GetMesh();
diff --git a/Engine/Source/Editor/Persona/Private/PersonaMeshDetails.h b/Engine/Source/Editor/Persona/Private/PersonaMeshDetails.h
--- a/Engine/Source/Editor/Persona/Private/PersonaMeshDetails.h
+++ b/Engine/Source/Editor/Persona/Private/PersonaMeshDetails.h
@@ -788,14 +788,6 @@ private:
 	/* Make uniform grid widget for Apex details */
 	TSharedRef<SUniformGridPanel> MakeClothingDetailsWidget(int32 AssetIndex) const;
 
-#if WITH_APEX_CLOTHING
-	/* Opens dialog to add a new clothing asset */
-	FReply OnOpenClothingFileClicked(IDetailLayoutBuilder* DetailLayout);
-
-	/* Reimports a clothing asset */ 
-	FReply OnReimportApexFileClicked(int32 AssetIndex, IDetailLayoutBuilder* DetailLayout);
-#endif
-
 	/* Removes a clothing asset */ 
 	FReply OnRemoveClothingAssetClicked(int32 AssetIndex, IDetailLayoutBuilder* DetailLayout);
 
diff --git a/Engine/Source/Editor/Persona/Private/PreviewSceneCustomizations.h b/Engine/Source/Editor/Persona/Private/PreviewSceneCustomizations.h
--- a/Engine/Source/Editor/Persona/Private/PreviewSceneCustomizations.h
+++ b/Engine/Source/Editor/Persona/Private/PreviewSceneCustomizations.h
@@ -12,7 +12,7 @@
 #include "AnimationEditorPreviewScene.h"
 
 #ifndef CHAOS_SIMULATION_DETAIL_VIEW_FACTORY_SELECTOR  // TODO: Decide whether to keep the detail view cloth selector after nvcloth has been removed
-#define CHAOS_SIMULATION_DETAIL_VIEW_FACTORY_SELECTOR WITH_CHAOS
+	#define CHAOS_SIMULATION_DETAIL_VIEW_FACTORY_SELECTOR 1
 #endif
 
 struct FAssetData;
diff --git a/Engine/Source/Editor/Persona/Private/SAnimViewportToolBar.cpp b/Engine/Source/Editor/Persona/Private/SAnimViewportToolBar.cpp
--- a/Engine/Source/Editor/Persona/Private/SAnimViewportToolBar.cpp
+++ b/Engine/Source/Editor/Persona/Private/SAnimViewportToolBar.cpp
@@ -717,7 +717,6 @@ TSharedRef<SWidget> SAnimViewportToolBar::GenerateCharacterMenu() const
 				})
 				);
 
-#if WITH_APEX_CLOTHING || WITH_CHAOS_CLOTHING
 			UDebugSkelMeshComponent* PreviewComp = Viewport.Pin()->GetPreviewScene()->GetPreviewMeshComponent();
 
 			if(PreviewComp)
@@ -727,7 +726,6 @@ TSharedRef<SWidget> SAnimViewportToolBar::GenerateCharacterMenu() const
 					LOCTEXT("CharacterMenu_ClothingSubMenuToolTip", "Options relating to clothing"),
 					FNewMenuDelegate::CreateRaw(const_cast<SAnimViewportToolBar*>(this), &SAnimViewportToolBar::FillCharacterClothingMenu));
 			}
-#endif // #if WITH_APEX_CLOTHING || WITH_CHAOS_CLOTHING
 		}
 
 		InMenuBuilder.AddSubMenu(
@@ -839,7 +837,6 @@ void SAnimViewportToolBar::OnMirrorDataTableSelected(const FAssetData& SelectedM
 
 void SAnimViewportToolBar::FillCharacterClothingMenu(FMenuBuilder& MenuBuilder)
 {
-#if WITH_APEX_CLOTHING || WITH_CHAOS_CLOTHING
 	const FAnimViewportShowCommands& Actions = FAnimViewportShowCommands::Get();
 
 	MenuBuilder.BeginSection("ClothPreview", LOCTEXT("ClothPreview_Label", "Simulation"));
@@ -882,7 +879,6 @@ void SAnimViewportToolBar::FillCharacterClothingMenu(FMenuBuilder& MenuBuilder)
 		}
 	}
 
-#endif // #if WITH_APEX_CLOTHING || WITH_CHAOS_CLOTHING
 }
 
 TSharedRef<SWidget> SAnimViewportToolBar::GenerateShowMenu() const
diff --git a/Engine/Source/Editor/Persona/Private/SAnimationEditorViewport.cpp b/Engine/Source/Editor/Persona/Private/SAnimationEditorViewport.cpp
--- a/Engine/Source/Editor/Persona/Private/SAnimationEditorViewport.cpp
+++ b/Engine/Source/Editor/Persona/Private/SAnimationEditorViewport.cpp
@@ -915,8 +915,6 @@ void SAnimationEditorViewportTabBody::BindCommands()
 
 	CommandList.EndGroup();
 
-#if WITH_APEX_CLOTHING || WITH_CHAOS_CLOTHING
-
 	//Clothing show options
 	CommandList.MapAction( 
 		ViewportShowMenuCommands.EnableClothSimulation,
@@ -963,7 +961,6 @@ void SAnimationEditorViewportTabBody::BindCommands()
 
 	CommandList.EndGroup();
 
-#endif// #if WITH_APEX_CLOTHING || WITH_CHAOS_CLOTHING
 
 	GetPreviewScene()->RegisterOnSelectedLODChanged(FOnSelectedLODChanged::CreateSP(this, &SAnimationEditorViewportTabBody::OnLODModelChanged));
 	//Bind LOD preview menu commands
@@ -1988,7 +1985,6 @@ bool SAnimationEditorViewportTabBody::CanUseProcessRootMotionMode(EProcessRootMo
 	return false;
 }
 
-#if WITH_APEX_CLOTHING || WITH_CHAOS_CLOTHING
 bool SAnimationEditorViewportTabBody::IsClothSimulationEnabled() const
 {
 	UDebugSkelMeshComponent* PreviewComponent = GetPreviewScene()->GetPreviewMeshComponent();
@@ -2168,7 +2164,6 @@ bool SAnimationEditorViewportTabBody::IsSectionsDisplayMode(ESectionDisplayMode
 {
 	return SectionsDisplayMode == DisplayMode;
 }
-#endif // #if WITH_APEX_CLOTHING || WITH_CHAOS_CLOTHING
 
 void SAnimationEditorViewportTabBody::AddRecordingNotification()
 {
diff --git a/Engine/Source/Editor/Persona/Private/SAnimationEditorViewport.h b/Engine/Source/Editor/Persona/Private/SAnimationEditorViewport.h
--- a/Engine/Source/Editor/Persona/Private/SAnimationEditorViewport.h
+++ b/Engine/Source/Editor/Persona/Private/SAnimationEditorViewport.h
@@ -513,7 +513,6 @@ public:
 
 	bool IsTurnTableSpeedSelected(int32 SpeedIndex) const;
 
-#if WITH_APEX_CLOTHING || WITH_CHAOS_CLOTHING
 	/** 
 	 * clothing show options 
 	*/
@@ -536,8 +535,6 @@ private:
 	void OnSetSectionsDisplayMode(ESectionDisplayMode DisplayMode);
 	bool IsSectionsDisplayMode(ESectionDisplayMode DisplayMode) const;
 
-#endif // #if WITH_APEX_CLOTHING || WITH_CHAOS_CLOTHING
-
 private:
 	/** Weak pointer back to the preview scene we are viewing */
 	TWeakPtr<class FAnimationEditorPreviewScene> PreviewScenePtr;
diff --git a/Engine/Source/Editor/PhysicsAssetEditor/Private/PhysicsAssetDetailsCustomization.cpp b/Engine/Source/Editor/PhysicsAssetEditor/Private/PhysicsAssetDetailsCustomization.cpp
--- a/Engine/Source/Editor/PhysicsAssetEditor/Private/PhysicsAssetDetailsCustomization.cpp
+++ b/Engine/Source/Editor/PhysicsAssetEditor/Private/PhysicsAssetDetailsCustomization.cpp
@@ -41,11 +41,6 @@ void FPhysicsAssetDetailsCustomization::CustomizeDetails(IDetailLayoutBuilder& D
 	PhysicalAnimationProfilesHandle = DetailLayout.GetProperty(GET_MEMBER_NAME_CHECKED(UPhysicsAsset, PhysicalAnimationProfiles));
 	ConstraintProfilesHandle = DetailLayout.GetProperty(GET_MEMBER_NAME_CHECKED(UPhysicsAsset, ConstraintProfiles));
 
-#if !WITH_CHAOS
-	// Hide Chaos-Only settings in PhysX
-	DetailLayout.GetProperty(GET_MEMBER_NAME_CHECKED(UPhysicsAsset, SolverIterations))->MarkHiddenByCustomization();
-#endif
-
 	DetailLayout.EditCategory(TEXT("Physical Animation Profiles"))
 	.AddProperty(PhysicalAnimationProfilesHandle)
 	.CustomWidget()
diff --git a/Engine/Source/Editor/PhysicsAssetEditor/Private/PhysicsAssetEditor.cpp b/Engine/Source/Editor/PhysicsAssetEditor/Private/PhysicsAssetEditor.cpp
--- a/Engine/Source/Editor/PhysicsAssetEditor/Private/PhysicsAssetEditor.cpp
+++ b/Engine/Source/Editor/PhysicsAssetEditor/Private/PhysicsAssetEditor.cpp
@@ -442,19 +442,11 @@ void FPhysicsAssetEditor::PostRedo( bool bSuccess )
 		{
 			FKConvexElem& Element = Body->AggGeom.ConvexElems[ElemIdx];
 
-#if PHYSICS_INTERFACE_PHYSX
-			if (Element.GetConvexMesh() == NULL)
-			{
-				bRecreate = true;
-				break;
-			}
-#elif WITH_CHAOS
 			if (Element.GetChaosConvexMesh() == NULL)
 			{
 				bRecreate = true;
 				break;
 			}
-#endif
 		}
 
 		if (bRecreate)
@@ -1370,14 +1362,13 @@ void FPhysicsAssetEditor::BuildMenuWidgetBody(FMenuBuilder& InMenuBuilder)
 				InSubMenuBuilder.AddMenuEntry( PhysicsAssetEditorCommands.DisableCollisionAll );
 				InSubMenuBuilder.EndSection();
 
-#if WITH_CHAOS
 				InSubMenuBuilder.BeginSection("CollisionFilteringHeader", LOCTEXT("CollisionFilteringHeader", "CollisionFiltering"));
 				InSubMenuBuilder.AddMenuEntry(PhysicsAssetEditorCommands.PrimitiveQueryAndPhysics);
 				InSubMenuBuilder.AddMenuEntry(PhysicsAssetEditorCommands.PrimitiveQueryOnly);
 				InSubMenuBuilder.AddMenuEntry(PhysicsAssetEditorCommands.PrimitivePhysicsOnly);
 				InSubMenuBuilder.AddMenuEntry(PhysicsAssetEditorCommands.PrimitiveNoCollision);
 				InSubMenuBuilder.EndSection();
-#endif
+
 				InSubMenuBuilder.BeginSection("MassHeader", LOCTEXT("MassHeader", "Mass"));
 				InSubMenuBuilder.AddMenuEntry(PhysicsAssetEditorCommands.PrimitiveContributeToMass);
 				InSubMenuBuilder.EndSection();
diff --git a/Engine/Source/Editor/PhysicsAssetEditor/Private/PhysicsAssetEditorAnimInstanceProxy.cpp b/Engine/Source/Editor/PhysicsAssetEditor/Private/PhysicsAssetEditorAnimInstanceProxy.cpp
--- a/Engine/Source/Editor/PhysicsAssetEditor/Private/PhysicsAssetEditorAnimInstanceProxy.cpp
+++ b/Engine/Source/Editor/PhysicsAssetEditor/Private/PhysicsAssetEditorAnimInstanceProxy.cpp
@@ -35,7 +35,6 @@ void FPhysicsAssetEditorAnimInstanceProxy::Initialize(UAnimInstance* InAnimInsta
 	FAnimPreviewInstanceProxy::Initialize(InAnimInstance);
 	ConstructNodes();
 
-#if WITH_CHAOS && !PHYSICS_INTERFACE_PHYSX
 	UPhysicsAsset* PhysicsAsset = InAnimInstance->GetSkelMeshComponent()->GetPhysicsAsset();
 	if (PhysicsAsset != nullptr)
 	{
@@ -44,7 +43,6 @@ void FPhysicsAssetEditorAnimInstanceProxy::Initialize(UAnimInstance* InAnimInsta
 	}
 
 	FloorActor = nullptr;
-#endif
 }
 
 void FPhysicsAssetEditorAnimInstanceProxy::ConstructNodes()
@@ -80,7 +78,6 @@ void FPhysicsAssetEditorAnimInstanceProxy::UpdateAnimationNode(const FAnimationU
 
 bool FPhysicsAssetEditorAnimInstanceProxy::Evaluate_WithRoot(FPoseContext& Output, FAnimNode_Base* InRootNode)
 {
-#if WITH_CHAOS && !PHYSICS_INTERFACE_PHYSX
 	ImmediatePhysics::FSimulation* Simulation = RagdollNode.GetSimulation();
 	if (Simulation != nullptr)
 	{
@@ -101,7 +98,6 @@ bool FPhysicsAssetEditorAnimInstanceProxy::Evaluate_WithRoot(FPoseContext& Outpu
 			SolverIterations.CollisionPushOutIterations
 		);
 	}
-#endif
 
 	if (CurrentAsset != nullptr)
 	{
@@ -126,7 +122,6 @@ void FPhysicsAssetEditorAnimInstanceProxy::AddImpulseAtLocation(FVector Impulse,
 
 void FPhysicsAssetEditorAnimInstanceProxy::Grab(FName InBoneName, const FVector& Location, const FRotator& Rotation, bool bRotationConstrained)
 {
-#if WITH_CHAOS && !PHYSICS_INTERFACE_PHYSX
 	ImmediatePhysics::FSimulation* Simulation = RagdollNode.GetSimulation();
 
 	if (TargetActor != nullptr)
@@ -152,12 +147,10 @@ void FPhysicsAssetEditorAnimInstanceProxy::Grab(FName InBoneName, const FVector&
 
 		HandleJoint = Simulation->CreateJoint(nullptr, TargetActor, HandleActor);
 	}
-#endif
 }
 
 void FPhysicsAssetEditorAnimInstanceProxy::Ungrab()
 {
-#if WITH_CHAOS && !PHYSICS_INTERFACE_PHYSX
 	if (TargetActor != nullptr)
 	{
 		ImmediatePhysics::FSimulation* Simulation = RagdollNode.GetSimulation();
@@ -167,33 +160,27 @@ void FPhysicsAssetEditorAnimInstanceProxy::Ungrab()
 		HandleActor = nullptr;
 		HandleJoint = nullptr;
 	}
-#endif
 }
 
 void FPhysicsAssetEditorAnimInstanceProxy::UpdateHandleTransform(const FTransform& NewTransform)
 {
-#if WITH_CHAOS && !PHYSICS_INTERFACE_PHYSX
 	if (HandleActor != nullptr)
 	{
 		HandleActor->SetKinematicTarget(NewTransform);
 	}
-#endif
 }
 
 void FPhysicsAssetEditorAnimInstanceProxy::UpdateDriveSettings(bool bLinearSoft, float LinearStiffness, float LinearDamping)
 {
-#if WITH_CHAOS && !PHYSICS_INTERFACE_PHYSX
 	using namespace Chaos;
 	if (HandleJoint != nullptr)
 	{
 		HandleJoint->SetSoftLinearSettings(bLinearSoft, LinearStiffness, LinearDamping);
 	}
-#endif
 }
 
 void FPhysicsAssetEditorAnimInstanceProxy::CreateSimulationFloor(FBodyInstance* FloorBodyInstance, const FTransform& Transform)
 {
-#if WITH_CHAOS && !PHYSICS_INTERFACE_PHYSX
 	using namespace Chaos;
 
 	DestroySimulationFloor();
@@ -207,12 +194,10 @@ void FPhysicsAssetEditorAnimInstanceProxy::CreateSimulationFloor(FBodyInstance*
 			Simulation->AddToCollidingPairs(FloorActor);
 		}
 	}
-#endif
 }
 
 void FPhysicsAssetEditorAnimInstanceProxy::DestroySimulationFloor()
 {
-#if WITH_CHAOS && !PHYSICS_INTERFACE_PHYSX
 	using namespace Chaos;
 	ImmediatePhysics::FSimulation* Simulation = RagdollNode.GetSimulation();
 	if ((Simulation != nullptr) && (FloorActor != nullptr))
@@ -220,6 +205,5 @@ void FPhysicsAssetEditorAnimInstanceProxy::DestroySimulationFloor()
 		Simulation->DestroyActor(FloorActor);
 		FloorActor = nullptr;
 	}
-#endif
 }
 
diff --git a/Engine/Source/Editor/PhysicsAssetEditor/Private/PhysicsAssetEditorSkeletalMeshComponent.cpp b/Engine/Source/Editor/PhysicsAssetEditor/Private/PhysicsAssetEditorSkeletalMeshComponent.cpp
--- a/Engine/Source/Editor/PhysicsAssetEditor/Private/PhysicsAssetEditorSkeletalMeshComponent.cpp
+++ b/Engine/Source/Editor/PhysicsAssetEditor/Private/PhysicsAssetEditorSkeletalMeshComponent.cpp
@@ -350,14 +350,10 @@ void UPhysicsAssetEditorSkeletalMeshComponent::RefreshBoneTransforms(FActorCompo
 
 void UPhysicsAssetEditorSkeletalMeshComponent::AddImpulseAtLocation(FVector Impulse, FVector Location, FName BoneName)
 {
-#if !WITH_CHAOS
-	Super::AddImpulseAtLocation(Impulse, Location, BoneName);
-#else
 	if (PreviewInstance != nullptr)
 	{
 		PreviewInstance->AddImpulseAtLocation(Impulse, Location, BoneName);
 	}
-#endif
 }
 
 bool UPhysicsAssetEditorSkeletalMeshComponent::ShouldCreatePhysicsState() const
diff --git a/Engine/Source/Editor/UnrealEd/Private/ApexClothingUtils.cpp b/Engine/Source/Editor/UnrealEd/Private/ApexClothingUtils.cpp
--- a/Engine/Source/Editor/UnrealEd/Private/ApexClothingUtils.cpp
+++ b/Engine/Source/Editor/UnrealEd/Private/ApexClothingUtils.cpp
@@ -54,99 +54,6 @@ void RefreshSkelMeshComponents(USkeletalMesh* SkelMesh)
 	}
 }
 
-#if WITH_APEX_CLOTHING
-FString PromptForClothingFile()
-{
-	if(IDesktopPlatform* Platform = FDesktopPlatformModule::Get())
-	{
-		const void* ParentWindowWindowHandle = FSlateApplication::Get().FindBestParentWindowHandleForDialogs(nullptr);
-
-		TArray<FString> OpenFilenames;
-		FString OpenFilePath = FEditorDirectories::Get().GetLastDirectory(ELastDirectory::MESH_IMPORT_EXPORT);
-
-		if(Platform->OpenFileDialog(
-			ParentWindowWindowHandle,
-			*LOCTEXT("ImportClothing_ChooseFile", "Choose clothing asset source file").ToString(),
-			OpenFilePath,
-			TEXT(""),
-			TEXT("APEX clothing asset(*.apx,*.apb)|*.apx;*.apb|All files (*.*)|*.*"),
-			EFileDialogFlags::None,
-			OpenFilenames
-			))
-		{
-			return OpenFilenames[0];
-		}
-	}
-
-	// Nothing picked, empty path
-	return FString();
-}
-
-void PromptAndImportClothing(USkeletalMesh* SkelMesh)
-{
-	ensure(SkelMesh);
-
-	FString Filename = PromptForClothingFile();
-
-	if(!Filename.IsEmpty())
-	{
-		FEditorDirectories::Get().SetLastDirectory(ELastDirectory::MESH_IMPORT_EXPORT, Filename);
-
-		UClothingAssetFactory* Factory = UClothingAssetFactory::StaticClass()->GetDefaultObject<UClothingAssetFactory>();  // TODO: Rename this class to UClothingAssetFactoryNv
-
-		if(Factory && Factory->CanImport(Filename))
-		{
-			Factory->Import(Filename, SkelMesh);
-		}
-	}
-}
-
-apex::ClothingAsset* CreateApexClothingAssetFromPxStream(physx::PxFileBuf& Stream)
-{
-	// Peek into the buffer to see what kind of data it is (binary or xml)
-	NvParameterized::Serializer::SerializeType SerializeType = GApexSDK->getSerializeType(Stream);
-	// Create an NvParameterized serializer for the correct data type
-	NvParameterized::Serializer* Serializer = GApexSDK->createSerializer(SerializeType);
-
-	if(!Serializer)
-	{
-		return NULL;
-	}
-	// Deserialize into a DeserializedData buffer
-	NvParameterized::Serializer::DeserializedData DeserializedData;
-	NvParameterized::Serializer::ErrorType Error = Serializer->deserialize(Stream, DeserializedData);
-
-	apex::Asset* ApexAsset = NULL;
-	if( DeserializedData.size() > 0 )
-	{
-		// The DeserializedData has something in it, so create an APEX asset from it
-		ApexAsset = GApexSDK->createAsset( DeserializedData[0], NULL);
-	}
-	// Release the serializer
-	Serializer->release();
-
-	return (apex::ClothingAsset*)ApexAsset;
-}
-
-apex::ClothingAsset* CreateApexClothingAssetFromBuffer(const uint8* Buffer, int32 BufferSize)
-{
-	apex::ClothingAsset* ApexClothingAsset = NULL;
-
-	// Wrap Buffer with the APEX read stream class
-	physx::PxFileBuf* Stream = GApexSDK->createMemoryReadStream(Buffer, BufferSize);
-
-	if (Stream != NULL)
-	{
-		ApexClothingAsset = ApexClothingUtils::CreateApexClothingAssetFromPxStream(*Stream);
-		// Release our stream
-		GApexSDK->releaseMemoryReadStream(*Stream);
-	}
-
-	return ApexClothingAsset;
-}
-
-#endif// #if WITH_APEX_CLOTHING
-
 void RestoreAllClothingSections(USkeletalMesh* SkelMesh, uint32 LODIndex, uint32 AssetIndex)
 {
 	if(FSkeletalMeshModel* Resource = SkelMesh->GetImportedModel())
@@ -180,23 +87,12 @@ void RemoveAssetFromSkeletalMesh(USkeletalMesh* SkelMesh, uint32 AssetIndex, boo
 		RestoreAllClothingSections(SkelMesh, LODIdx, AssetIndex);
 	}
 
-#if WITH_APEX_CLOTHING
-	apex::ClothingAsset* ApexClothingAsset = SkelMesh->ClothingAssets_DEPRECATED[AssetIndex].ApexClothingAsset;	//Can't delete apex asset until after apex actors so we save this for now and reregister component (which will trigger the actor delete)
-#endif	// WITH_APEX_CLOTHING
 	SkelMesh->ClothingAssets_DEPRECATED.RemoveAt(AssetIndex);	//have to remove the asset from the array so that new actors are not created for asset pending deleting
 
 	SkelMesh->PostEditChange(); // update derived data
 
 	ReregisterSkelMeshComponents(SkelMesh);
 
-#if WITH_APEX_CLOTHING
-	if(bReleaseAsset)
-	{
-		//Now we can actually delete the asset
-		GPhysCommandHandler->DeferredRelease(ApexClothingAsset);
-	}
-#endif	// WITH_APEX_CLOTHING
-
 	if(bRecreateSkelMeshComponent)
 	{
 		// Refresh skeletal mesh components
diff --git a/Engine/Source/Editor/UnrealEd/Private/EditorEngine.cpp b/Engine/Source/Editor/UnrealEd/Private/EditorEngine.cpp
--- a/Engine/Source/Editor/UnrealEd/Private/EditorEngine.cpp
+++ b/Engine/Source/Editor/UnrealEd/Private/EditorEngine.cpp
@@ -245,10 +245,7 @@
 #include "DerivedDataBuildRemoteExecutor.h"
 #include "DerivedDataBuildWorkers.h"
 #include "AssetCompilingManager.h"
-
-#if WITH_CHAOS
 #include "ChaosSolversModule.h"
-#endif
 
 #include "DeviceProfiles/DeviceProfile.h"
 #include "DeviceProfiles/DeviceProfileManager.h"
@@ -4301,13 +4298,6 @@ void UEditorEngine::CleanupPhysicsSceneThatWasInitializedForSave(UWorld* World,
 
 	World->SetPhysicsScene(nullptr);
 
-#if PHYSICS_INTERFACE_PHYSX
-	if (GPhysCommandHandler)
-	{
-		GPhysCommandHandler->Flush();
-	}
-#endif // PHYSICS_INTERFACE_PHYSX
-
 	// Update components again in case it was a world without a physics scene but did have rendered components.
 	World->UpdateWorldComponents(true, true);
 
diff --git a/Engine/Source/Editor/UnrealEd/Private/EditorPhysXSupport.h b/Engine/Source/Editor/UnrealEd/Private/EditorPhysXSupport.h
--- a/Engine/Source/Editor/UnrealEd/Private/EditorPhysXSupport.h
+++ b/Engine/Source/Editor/UnrealEd/Private/EditorPhysXSupport.h
@@ -5,76 +5,3 @@
 =============================================================================*/
 
 #pragma once
-
-#include "CoreMinimal.h"
-#include "EngineDefines.h"
-
-#if WITH_PHYSX
-
-#pragma warning( push )
-#pragma warning( disable : 4946 ) // reinterpret_cast used between related classes
-
-PRAGMA_DISABLE_SHADOW_VARIABLE_WARNINGS
-
-#if USING_CODE_ANALYSIS
-	#pragma warning( push )
-	#pragma warning( disable : ALL_CODE_ANALYSIS_WARNINGS )
-#endif	// USING_CODE_ANALYSIS
-
-// different packing for Intel 32-bit on Linux
-#if (PLATFORM_LINUX && PLATFORM_CPU_X86_FAMILY && !PLATFORM_64BITS)
-	#pragma pack(push,16)
-#else
-	#pragma pack(push,8)
-#endif
-THIRD_PARTY_INCLUDES_START
-
-#include "Px.h"
-#include "PxPhysicsAPI.h"
-#include "PxRenderBuffer.h"
-#include "PxExtensionsAPI.h"
-#include "PxPvd.h"
-//#include "PxDefaultCpuDispatcher.h"
-
-// vehicle related header files
-//#include "PxVehicleSDK.h"
-//#include "PxVehicleUtils.h"
-
-// utils
-#include "PxGeometryQuery.h"
-#include "PxMeshQuery.h"
-#include "PxTriangle.h"
-
-// APEX
-#if WITH_APEX
-// Framework
-#include "Apex.h"
-
-// Modules
-#include "ModuleClothing.h"
-
-// Assets
-#include "Assets/ClothingAssetNv.h"
-
-// Actors
-#include "ClothingActor.h"
-
-// Utilities
-#include "NvParamUtils.h"
-
-#endif // #if WITH_APEX
-
-THIRD_PARTY_INCLUDES_END
-#pragma pack(pop)
-
-#if USING_CODE_ANALYSIS
-	#pragma warning( pop )
-#endif	// USING_CODE_ANALYSIS
-
-#pragma warning( pop )
-
-PRAGMA_ENABLE_SHADOW_VARIABLE_WARNINGS
-
-using namespace physx;
-
-#endif // WITH_PHYSX
diff --git a/Engine/Source/Editor/UnrealEd/Private/Fbx/FbxMainExport.cpp b/Engine/Source/Editor/UnrealEd/Private/Fbx/FbxMainExport.cpp
--- a/Engine/Source/Editor/UnrealEd/Private/Fbx/FbxMainExport.cpp
+++ b/Engine/Source/Editor/UnrealEd/Private/Fbx/FbxMainExport.cpp
@@ -96,17 +96,11 @@
 #include "MovieSceneTimeHelpers.h"
 #include "DynamicMeshBuilder.h"
 
-#if PHYSICS_INTERFACE_PHYSX
-#include "PhysXPublic.h"
-#include "geometry/PxConvexMesh.h"
-#include "PhysicsEngine/BodySetup.h"
-#endif // WITH_PHYSX
-
-#if WITH_CHAOS
 #include "Chaos/Core.h"
 #include "Chaos/Particles.h"
 #include "Chaos/Plane.h"
-#endif
+#include "ChaosCheck.h"
+#include "Chaos/Convex.h"
 
 #include "Exporters/FbxExportOption.h"
 #include "FbxExportOptionsWindow.h"
@@ -115,11 +109,6 @@
 #include "Interfaces/IMainFrameModule.h"
 #include "UObject/MetaData.h"
 
-#if WITH_CHAOS
-#include "ChaosCheck.h"
-#include "Chaos/Convex.h"
-#endif
-
 namespace UnFbx
 {
 
@@ -3659,8 +3648,6 @@ void DetermineVertsToWeld(TArray<int32>& VertRemap, TArray<int32>& UniqueVerts,
 	}
 }
 
-#if (WITH_PHYSX && PHYSICS_INTERFACE_PHYSX) || WITH_CHAOS
-
 class FCollisionFbxExporter
 {
 public:
@@ -3827,11 +3814,7 @@ public:
 private:
 	uint32 GetConvexVerticeNumber(const FKConvexElem &ConvexElem)
 	{
-#if WITH_PHYSX && PHYSICS_INTERFACE_PHYSX
-		return ConvexElem.GetConvexMesh() != nullptr ? ConvexElem.GetConvexMesh()->getNbVertices() : 0;
-#elif WITH_CHAOS
 		return ConvexElem.VertexData.Num();
-#endif
 	}
 
 	uint32 GetBoxVerticeNumber() { return 24; }
@@ -3842,20 +3825,6 @@ private:
 
 	void AddConvexVertex(const FKConvexElem &ConvexElem)
 	{
-#if WITH_PHYSX && PHYSICS_INTERFACE_PHYSX
-		const physx::PxConvexMesh* ConvexMesh = ConvexElem.GetConvexMesh();
-		if (ConvexMesh == nullptr)
-		{
-			return;
-		}
-		const PxVec3 *VertexArray = ConvexMesh->getVertices();
-		for (uint32 PosIndex = 0; PosIndex < ConvexMesh->getNbVertices(); ++PosIndex)
-		{
-			FVector Position = P2UVector(VertexArray[PosIndex]);
-			ControlPoints[CurrentVertexOffset + PosIndex] = FbxVector4(Position.X, -Position.Y, Position.Z);
-		}
-		CurrentVertexOffset += ConvexMesh->getNbVertices();
-#elif WITH_CHAOS
 		const TArray<FVector>& VertexArray = ConvexElem.VertexData;
 		for (int32 PosIndex = 0; PosIndex < VertexArray.Num(); ++PosIndex)
 		{
@@ -3863,36 +3832,10 @@ private:
 			ControlPoints[CurrentVertexOffset + PosIndex] = FbxVector4(Position.X, -Position.Y, Position.Z);
 		}
 		CurrentVertexOffset += VertexArray.Num();
-#endif
 	}
 
 	void AddConvexNormals(const FKConvexElem &ConvexElem)
 	{
-#if WITH_PHYSX && PHYSICS_INTERFACE_PHYSX
-		const physx::PxConvexMesh* ConvexMesh = ConvexElem.GetConvexMesh();
-		if (ConvexMesh == nullptr)
-		{
-			return;
-		}
-		const PxU8* PIndexBuffer = ConvexMesh->getIndexBuffer();
-		int32 PolygonNumber = ConvexMesh->getNbPolygons();
-		for (int32 PolyIndex = 0; PolyIndex < PolygonNumber; ++PolyIndex)
-		{
-			PxHullPolygon PolyData;
-			if (!ConvexMesh->getPolygonData(PolyIndex, PolyData))
-			{
-				continue;
-			}
-			const PxVec3 PPlaneNormal(PolyData.mPlane[0], PolyData.mPlane[1], PolyData.mPlane[2]);
-			FVector Normal = P2UVector(PPlaneNormal.getNormalized());
-			FbxVector4 FbxNormal = FbxVector4(Normal.X, -Normal.Y, Normal.Z);
-			// add vertices 
-			for (PxU32 j = 0; j < PolyData.mNbVerts; ++j)
-			{
-				LayerElementNormal->GetDirectArray().Add(FbxNormal);
-			}
-		}
-#elif WITH_CHAOS
 		const auto& ConvexMesh = ConvexElem.GetChaosConvexMesh();
 		if (!ConvexMesh.IsValid())
 		{
@@ -3909,38 +3852,10 @@ private:
 				LayerElementNormal->GetDirectArray().Add(FbxNormal);
 			}
 		}
-#endif
 	}
 
 	void AddConvexPolygon(const FKConvexElem &ConvexElem)
 	{
-#if WITH_PHYSX && PHYSICS_INTERFACE_PHYSX
-		const physx::PxConvexMesh* ConvexMesh = ConvexElem.GetConvexMesh();
-		if (ConvexMesh == nullptr)
-		{
-			return;
-		}
-		const PxU8* PIndexBuffer = ConvexMesh->getIndexBuffer();
-		int32 PolygonNumber = ConvexMesh->getNbPolygons();
-		for (int32 PolyIndex = 0; PolyIndex < PolygonNumber; ++PolyIndex)
-		{
-			PxHullPolygon PolyData;
-			if (!ConvexMesh->getPolygonData(PolyIndex, PolyData))
-			{
-				continue;
-			}
-			Mesh->BeginPolygon(ActualMatIndex);
-			const PxU8* PolyIndices = PIndexBuffer + PolyData.mIndexBase;
-			// add vertices 
-			for (PxU32 j = 0; j < PolyData.mNbVerts; ++j)
-			{
-				const uint32 VertIndex = CurrentVertexOffset + PolyIndices[j];
-				Mesh->AddPolygon(VertIndex);
-			}
-			Mesh->EndPolygon();
-		}
-		CurrentVertexOffset += ConvexMesh->getNbVertices();
-#elif WITH_CHAOS
 		const auto& ConvexMesh = ConvexElem.GetChaosConvexMesh();
 		if (!ConvexMesh.IsValid())
 		{
@@ -3962,7 +3877,6 @@ private:
 		}
 
 		CurrentVertexOffset += ConvexElem.VertexData.Num();
-#endif
 	}
 
 	void AddBoxVertex(const FKBoxElem &BoxElem)
@@ -4397,7 +4311,7 @@ FbxNode* FFbxExporter::ExportCollisionMesh(const UStaticMesh* StaticMesh, const
 	FbxActor->SetNodeAttribute(CollisionMesh);
 	return FbxActor;
 }
-#endif //WITH_PHYSX
+
 
 void FFbxExporter::ExportObjectMetadata(const UObject* ObjectToExport, FbxNode* Node)
 {
@@ -4787,13 +4701,10 @@ FbxNode* FFbxExporter::ExportStaticMeshToFbx(const UStaticMesh* StaticMesh, int3
 		}
 	}
 
-#if (WITH_PHYSX && PHYSICS_INTERFACE_PHYSX) || WITH_CHAOS
 	if ((ExportLOD == 0 || ExportLOD == -1) && GetExportOptions()->Collision)
 	{
 		ExportCollisionMesh(StaticMesh, MeshName, FbxActor);
 	}
-#endif
-
 
 	//Set the original meshes in case it was already existing
 	FbxActor->SetNodeAttribute(Mesh);
diff --git a/Engine/Source/Editor/UnrealEd/Private/FbxExporter.h b/Engine/Source/Editor/UnrealEd/Private/FbxExporter.h
--- a/Engine/Source/Editor/UnrealEd/Private/FbxExporter.h
+++ b/Engine/Source/Editor/UnrealEd/Private/FbxExporter.h
@@ -325,10 +325,8 @@ private:
 	void ExportAnimTrack( IAnimTrackAdapter& AnimTrackAdapter, AActor* Actor, USkeletalMeshComponent* SkeletalMeshComponent, float SamplingRate );
 
 	void ExportModel(UModel* Model, FbxNode* Node, const char* Name);
-	
-#if (WITH_PHYSX && PHYSICS_INTERFACE_PHYSX) || WITH_CHAOS
+
 	FbxNode* ExportCollisionMesh(const UStaticMesh* StaticMesh, const TCHAR* MeshName, FbxNode* ParentActor);
-#endif
 
 	/**
 	 * Exports a static mesh
diff --git a/Engine/Source/Editor/UnrealEd/Private/FbxMeshUtils.cpp b/Engine/Source/Editor/UnrealEd/Private/FbxMeshUtils.cpp
--- a/Engine/Source/Editor/UnrealEd/Private/FbxMeshUtils.cpp
+++ b/Engine/Source/Editor/UnrealEd/Private/FbxMeshUtils.cpp
@@ -32,11 +32,6 @@
 #include "InterchangeProjectSettings.h"
 #include "Settings/EditorExperimentalSettings.h"
 
-#if WITH_APEX_CLOTHING
-	#include "ApexClothingUtils.h"
-#endif // #if WITH_APEX_CLOTHING
-
-
 #include "Misc/FbxErrors.h"
 #include "Framework/Notifications/NotificationManager.h"
 #include "Widgets/Notifications/SNotificationList.h"
diff --git a/Engine/Source/Editor/UnrealEd/Public/ApexClothingUtils.h b/Engine/Source/Editor/UnrealEd/Public/ApexClothingUtils.h
--- a/Engine/Source/Editor/UnrealEd/Public/ApexClothingUtils.h
+++ b/Engine/Source/Editor/UnrealEd/Public/ApexClothingUtils.h
@@ -7,30 +7,8 @@
 
 class USkeletalMesh;
 
-#if WITH_APEX_CLOTHING
-namespace nvidia
-{
-	namespace apex
-	{
-		class ClothingAsset;
-	}
-}
-#endif  // #if WITH_APEX_CLOTHING
-
-// Define interface for importing apex clothing asset
 namespace ApexClothingUtils
 {
-#if WITH_APEX_CLOTHING
-	// Prompt the user to select an APEX file
-	UNREALED_API FString PromptForClothingFile();
-
-	// Prompt the user to select an APEX file that will be imported to the new UClothingAssetCommon format
-	UNREALED_API void PromptAndImportClothing(USkeletalMesh* SkelMesh);
-
-	// Given a buffer, build an apex clothing asset
-	UNREALED_API nvidia::apex::ClothingAsset* CreateApexClothingAssetFromBuffer(const uint8* Buffer, int32 BufferSize);
-#endif  // #if WITH_APEX_CLOTHING
-
 	// Functions below remain from previous clothing system and only remain to remove the bound
 	// data from a skeletal mesh. This is done when postloading USkeletalMesh when upgrading the assets
 
diff --git a/Engine/Source/Programs/DerivedDataBuildWorker/DerivedDataBuildWorker.Build.cs b/Engine/Source/Programs/DerivedDataBuildWorker/DerivedDataBuildWorker.Build.cs
--- a/Engine/Source/Programs/DerivedDataBuildWorker/DerivedDataBuildWorker.Build.cs
+++ b/Engine/Source/Programs/DerivedDataBuildWorker/DerivedDataBuildWorker.Build.cs
@@ -23,8 +23,6 @@ public abstract class DerivedDataBuildWorkerTarget : TargetRules
 		bCompileNavmeshSegmentLinks		= false;
 		bCompileRecast					= false;
 		bCompileICU 					= false;
-		bCompileChaos					= false;
-		bUseChaos						= false;
 		bWithLiveCoding					= false;
 		bBuildDeveloperTools			= false;
 		bUseMallocProfiler				= false;
diff --git a/Engine/Source/Programs/Enterprise/Datasmith/DatasmithCADWorker/DatasmithCADWorker.Target.cs b/Engine/Source/Programs/Enterprise/Datasmith/DatasmithCADWorker/DatasmithCADWorker.Target.cs
--- a/Engine/Source/Programs/Enterprise/Datasmith/DatasmithCADWorker/DatasmithCADWorker.Target.cs
+++ b/Engine/Source/Programs/Enterprise/Datasmith/DatasmithCADWorker/DatasmithCADWorker.Target.cs
@@ -34,9 +34,6 @@ public class DatasmithCADWorkerTarget : TargetRules
 
 		bLegalToDistributeBinary = true;
 
-		bCompilePhysX = false;
-		bCompileAPEX = false;
-		bCompileNvCloth = false;
 		bCompileICU = false;
 		bCompileCEF3 = false;
 
diff --git a/Engine/Source/Programs/HeadlessChaos/HeadlessChaos.Target.cs b/Engine/Source/Programs/HeadlessChaos/HeadlessChaos.Target.cs
--- a/Engine/Source/Programs/HeadlessChaos/HeadlessChaos.Target.cs
+++ b/Engine/Source/Programs/HeadlessChaos/HeadlessChaos.Target.cs
@@ -22,14 +22,6 @@ public class HeadlessChaosTarget : TestTargetRules
 
 		GlobalDefinitions.Add("CHAOS_SERIALIZE_OUT=1");
 
-		// Force enable Chaos as the physics engine for this project as this
-		// is a Chaos unit test framework - it doesn't build with PhysX enabled
-		bUseChaos = true;
-		bCompileChaos = true;
-		bCompilePhysX = false;
-		bCompileAPEX = false;
-		bCompileNvCloth = false;
-
 		// our gtest does not have debug crt libs
 		bDebugBuildsActuallyUseDebugCRT = false;
 	}
diff --git a/Engine/Source/Programs/HeadlessChaos/Private/GeometryCollection/GeometryCollectionTestUtility.cpp b/Engine/Source/Programs/HeadlessChaos/Private/GeometryCollection/GeometryCollectionTestUtility.cpp
--- a/Engine/Source/Programs/HeadlessChaos/Private/GeometryCollection/GeometryCollectionTestUtility.cpp
+++ b/Engine/Source/Programs/HeadlessChaos/Private/GeometryCollection/GeometryCollectionTestUtility.cpp
@@ -240,40 +240,5 @@ namespace GeometryCollectionTest {
 		PhysicalMaterial->DisabledAngularThreshold = 0;
 	}
 
-	/*
-	FGeometryCollectionPhysicsProxy* RigidBodySetup(
-		TUniquePtr<Chaos::FChaosPhysicsMaterial> &PhysicalMaterial,
-		TSharedPtr<FGeometryCollection>& RestCollection,
-		TSharedPtr<FGeometryDynamicCollection>& DynamicCollection,
-		FInitFunc CustomFunc
-	)
-	{
-#if INCLUDE_CHAOS
-		auto InitFunc = [&RestCollection, &DynamicCollection, &PhysicalMaterial, &CustomFunc](FSimulationParameters& InParams)
-		{
-			InParams.RestCollection = RestCollection.Get();
-			InParams.DynamicCollection = DynamicCollection.Get();
-			InParams.PhysicalMaterial = MakeSerializable(PhysicalMaterial);
-			InParams.Shared.SizeSpecificData[0].CollisionType = ECollisionTypeEnum::Chaos_Volumetric;
-
-			if (CustomFunc != nullptr)
-			{
-				CustomFunc(InParams);
-			}
-
-			InParams.Simulating = true;
-			Chaos::FErrorReporter ErrorReporter;
-			BuildSimulationData(ErrorReporter, *RestCollection, InParams.Shared);
-		};
-
-		FGeometryCollectionPhysicsProxy* PhysObject = new FGeometryCollectionPhysicsProxy(nullptr, DynamicCollection.Get(), InitFunc, nullptr, nullptr);;
-		PhysObject->Initialize();
-		return PhysObject;
-#else
-		return nullptr;
-#endif
-	}
-	*/
-
 } // end namespace GeometryCollectionTest
 
diff --git a/Engine/Source/Programs/HeadlessChaosPerf/HeadlessChaosPerf.Target.cs b/Engine/Source/Programs/HeadlessChaosPerf/HeadlessChaosPerf.Target.cs
--- a/Engine/Source/Programs/HeadlessChaosPerf/HeadlessChaosPerf.Target.cs
+++ b/Engine/Source/Programs/HeadlessChaosPerf/HeadlessChaosPerf.Target.cs
@@ -29,13 +29,5 @@ public class HeadlessChaosPerfTarget : TargetRules
 
 		GlobalDefinitions.Add("CHAOS_SERIALIZE_OUT=1");
 		GlobalDefinitions.Add("CSV_PROFILER=1");
-
-		// Force enable Chaos as the physics engine for this project as this
-		// is a Chaos test - it will not build with PhysX enabled
-		bUseChaos = true;
-		bCompileChaos = true;
-		bCompilePhysX = false;
-		bCompileAPEX = false;
-		bCompileNvCloth = false;
 	}
 }
diff --git a/Engine/Source/Programs/InterchangeWorker/InterchangeWorker.Target.cs b/Engine/Source/Programs/InterchangeWorker/InterchangeWorker.Target.cs
--- a/Engine/Source/Programs/InterchangeWorker/InterchangeWorker.Target.cs
+++ b/Engine/Source/Programs/InterchangeWorker/InterchangeWorker.Target.cs
@@ -34,9 +34,6 @@ public class InterchangeWorkerTarget : TargetRules
 
 		bLegalToDistributeBinary = true;
 
-		bCompilePhysX = false;
-		bCompileAPEX = false;
-		bCompileNvCloth = false;
 		bCompileICU = false;
 		bCompileCEF3 = false;
 
diff --git a/Engine/Source/Programs/UnrealBuildTool/Configuration/ModuleRules.cs b/Engine/Source/Programs/UnrealBuildTool/Configuration/ModuleRules.cs
--- a/Engine/Source/Programs/UnrealBuildTool/Configuration/ModuleRules.cs
+++ b/Engine/Source/Programs/UnrealBuildTool/Configuration/ModuleRules.cs
@@ -1222,163 +1222,48 @@ namespace UnrealBuildTool
 		/// </summary>
 		public void SetupModulePhysicsSupport(ReadOnlyTargetRules Target)
 		{
-			PublicDependencyModuleNames.Add("PhysicsCore");
+			PublicIncludePathModuleNames.AddRange(
+					new string[] {
+					"Chaos",
+					}
+				);
 
-			bool bUseNonPhysXInterface = Target.bUseChaos == true;
 			PublicDependencyModuleNames.AddRange(
 				new string[] {
+					"PhysicsCore",
 					"Chaos",
-				}
+					}
 				);
-			// 
-			if (Target.bCompileChaos == true || Target.bUseChaos == true)
-            {
-                PublicDefinitions.Add("INCLUDE_CHAOS=1");
-			}
-            else
-            {
-                PublicDefinitions.Add("INCLUDE_CHAOS=0");
-            }
-            // definitions used outside of PhysX/APEX need to be set here, not in PhysX.Build.cs or APEX.Build.cs, 
-            // since we need to make sure we always set it, even to 0 (because these are Private dependencies, the
-            // defines inside their Build.cs files won't leak out)
-            if (Target.bCompilePhysX == true && Target.bCompileChaos == false && Target.bUseChaos == false)
-			{
-				PrivateDependencyModuleNames.Add("PhysX");
-			}
 
-			if(Target.bCompileChaos || Target.bUseChaos || Target.bCompilePhysX)
-			{
-				PublicDefinitions.Add("WITH_PHYSX=1");
-			}
-			else
-			{
-				PublicDefinitions.Add("WITH_PHYSX=0");
-			}
+			PublicDefinitions.Add("WITH_CLOTH_COLLISION_DETECTION=1");
 
-			if(!bUseNonPhysXInterface)
+			if(!bTreatAsEngineModule)
 			{
-				// Disable non-physx interfaces
-				PublicDefinitions.Add("WITH_CHAOS=0");
-				PublicDefinitions.Add("WITH_CHAOS_CLOTHING=0");
-
-				// 
-				// WITH_CHAOS_NEEDS_TO_BE_FIXED
-				//
-				// Anything wrapped in this define needs to be fixed
-				// in one of the build targets. This define was added
-				// to help identify complier failures between the
-				// the three build targets( UseChaos, PhysX, WithChaos )
-				// This defaults to off , and will be enabled for bUseChaos. 
-				// This define should be removed when all the references 
-				// have been fixed across the different builds. 
-				//
-				PublicDefinitions.Add("WITH_CHAOS_NEEDS_TO_BE_FIXED=0");
-
-				if (Target.bCompilePhysX)
-				{
-					PublicDefinitions.Add("PHYSICS_INTERFACE_PHYSX=1");
-				}
-				else
-				{
-					PublicDefinitions.Add("PHYSICS_INTERFACE_PHYSX=0");
-				}
-
-				if (Target.bCompileAPEX == true)
-				{
-					if (!Target.bCompilePhysX)
-					{
-						throw new BuildException("APEX is enabled, without PhysX. This is not supported!");
-					}
-					PrivateDependencyModuleNames.Add("APEX");
-					PublicDefinitions.Add("WITH_APEX=1");
-
-					if (Target.Platform.ToString() == "HoloLens") //  Do not use Apex Cloth for HoloLens.  TODO: can we enable this in the future?
-					{
-						PublicDefinitions.Add("WITH_APEX_CLOTHING=0");
-						PublicDefinitions.Add("WITH_CLOTH_COLLISION_DETECTION=0");
-					}
-					else
-					{
-						PublicDefinitions.Add("WITH_APEX_CLOTHING=1");
-						PublicDefinitions.Add("WITH_CLOTH_COLLISION_DETECTION=1");
-					}
-
-					PublicDefinitions.Add("WITH_PHYSX_COOKING=1");  // APEX currently relies on cooking even at runtime
-
-				}
-				else
-				{
-					PublicDefinitions.Add("WITH_APEX=0");
-					PublicDefinitions.Add("WITH_APEX_CLOTHING=0");
-					PublicDefinitions.Add("WITH_CLOTH_COLLISION_DETECTION=0");
-					PublicDefinitions.Add(string.Format("WITH_PHYSX_COOKING={0}", Target.Type == TargetType.Editor && Target.bCompilePhysX ? 1 : 0));  // without APEX, we only need cooking in editor builds
-				}
-
-				if (Target.bCompileNvCloth == true)
-				{
-					if (!Target.bCompilePhysX)
-					{
-						throw new BuildException("NvCloth is enabled, without PhysX. This is not supported!");
-					}
+				// Non-engine modules may still be relying on appropriate definitions for physics.
+				// Nothing in engine should use these anymore as they were all deprecated and 
+				// assumed to be in the following configuration from 5.1
 
-					PrivateDependencyModuleNames.Add("NvCloth");
-					PublicDefinitions.Add("WITH_NVCLOTH=1");
-
-				}
-				else
-				{
-					PublicDefinitions.Add("WITH_NVCLOTH=0");
-				}
-			}
-			else
-			{
-				// Disable apex/cloth/physx interface
-				PublicDefinitions.Add("PHYSICS_INTERFACE_PHYSX=0");
-				PublicDefinitions.Add("WITH_APEX=0");
-				PublicDefinitions.Add("WITH_APEX_CLOTHING=0");
-				PublicDefinitions.Add(string.Format("WITH_PHYSX_COOKING={0}", Target.Type == TargetType.Editor && Target.bCompilePhysX ? 1 : 0));  // without APEX, we only need cooking in editor builds
-				PublicDefinitions.Add("WITH_NVCLOTH=0");
-
-				if(Target.bUseChaos)
+				Func<string, string, string, string> GetDeprecatedPhysicsMacro = (string Macro, string Value, string Version) =>
 				{
-					PublicDefinitions.Add("WITH_CHAOS=1");
-					PublicDefinitions.Add("WITH_CHAOS_NEEDS_TO_BE_FIXED=1");
-					PublicDefinitions.Add("WITH_CHAOS_CLOTHING=1");
-					PublicDefinitions.Add("WITH_CLOTH_COLLISION_DETECTION=1");
-					
-					PublicIncludePathModuleNames.AddRange(
-						new string[] {
-						"Chaos",
-						}
-					);
-
-					PublicDependencyModuleNames.AddRange(
-						new string[] {
-						"Chaos",
-						}
-					);
-				}
-				else
-				{
-					PublicDefinitions.Add("WITH_CHAOS=0");
-					PublicDefinitions.Add("WITH_CHAOS_NEEDS_TO_BE_FIXED=0");
-					PublicDefinitions.Add("WITH_CHAOS_CLOTHING=0");
-					PublicDefinitions.Add("WITH_CLOTH_COLLISION_DETECTION=0");
-				}
+					return Macro + "=UE_DEPRECATED_MACRO(" + Version + ", \"" + Macro + " is deprecated and should always be considered " + Value + ".\") " + Value;
+				};
+
+				PublicDefinitions.AddRange(
+					new string[]{
+						GetDeprecatedPhysicsMacro("INCLUDE_CHAOS", "1", "5.1"),
+						GetDeprecatedPhysicsMacro("WITH_CHAOS", "1", "5.1"),
+						GetDeprecatedPhysicsMacro("WITH_CHAOS_CLOTHING", "1", "5.1"),
+						GetDeprecatedPhysicsMacro("WITH_CHAOS_NEEDS_TO_BE_FIXED", "1", "5.1"),
+						GetDeprecatedPhysicsMacro("WITH_PHYSX", "1", "5.1"),
+						GetDeprecatedPhysicsMacro("WITH_PHYSX_COOKING", "0", "5.1"),
+						GetDeprecatedPhysicsMacro("PHYSICS_INTERFACE_PHYSX", "0", "5.1"),
+						GetDeprecatedPhysicsMacro("WITH_APEX", "0", "5.1"),
+						GetDeprecatedPhysicsMacro("WITH_APEX_CLOTHING", "0", "5.1"),
+						GetDeprecatedPhysicsMacro("WITH_NVCLOTH", "0", "5.1"),
+						GetDeprecatedPhysicsMacro("WITH_IMMEDIATE_PHYSX", "0", "5.1"),
+						GetDeprecatedPhysicsMacro("WITH_CUSTOM_SQ_STRUCTURE", "0", "5.1")
+					});
 			}
-
-			if(Target.bCustomSceneQueryStructure)
-			{
-				PublicDefinitions.Add("WITH_CUSTOM_SQ_STRUCTURE=1");
-			}
-			else
-			{
-				PublicDefinitions.Add("WITH_CUSTOM_SQ_STRUCTURE=0");
-			}
-
-			// Unused interface
-			PublicDefinitions.Add("WITH_IMMEDIATE_PHYSX=0");
 		}
 
 		/// <summary>
diff --git a/Engine/Source/Programs/UnrealBuildTool/Configuration/TargetRules.cs b/Engine/Source/Programs/UnrealBuildTool/Configuration/TargetRules.cs
--- a/Engine/Source/Programs/UnrealBuildTool/Configuration/TargetRules.cs
+++ b/Engine/Source/Programs/UnrealBuildTool/Configuration/TargetRules.cs
@@ -572,6 +572,7 @@ namespace UnrealBuildTool
 		[RequiresUniqueBuildEnvironment]
 		[CommandLine("-NoCompileChaos", Value = "false")]
 		[CommandLine("-CompileChaos", Value = "true")]
+		[Obsolete("Deprecated in UE5.1 - No longer used as Chaos is always enabled.")]
 		public bool bCompileChaos = true;
 
 		/// <summary>
@@ -580,6 +581,7 @@ namespace UnrealBuildTool
 		[RequiresUniqueBuildEnvironment]
 		[CommandLine("-NoUseChaos", Value = "false")]
 		[CommandLine("-UseChaos", Value = "true")]
+		[Obsolete("Deprecated in UE5.1 - No longer used as Chaos is always enabled.")]
 		public bool bUseChaos = true;
 
 		/// <summary>
@@ -598,12 +600,14 @@ namespace UnrealBuildTool
 		/// Whether scene query acceleration is done by UE. The physx scene query structure is still created, but we do not use it.
 		/// </summary>
 		[RequiresUniqueBuildEnvironment]
+		[Obsolete("Deprecated in UE5.1 - No longer used in engine.")]
 		public bool bCustomSceneQueryStructure = false;
 
 		/// <summary>
 		/// Whether to include PhysX support.
 		/// </summary>
 		[RequiresUniqueBuildEnvironment]
+		[Obsolete("Deprecated in UE5.1 - No longer used as Chaos is always enabled.")]
 		public bool bCompilePhysX = false;
 
 		/// <summary>
@@ -611,12 +615,14 @@ namespace UnrealBuildTool
 		/// </summary>
 		[RequiresUniqueBuildEnvironment]
 		[ConfigFile(ConfigHierarchyType.Engine, "/Script/BuildSettings.BuildSettings", "bCompileApex")]
+		[Obsolete("Deprecated in UE5.1 - No longer used as Chaos is always enabled.")]
 		public bool bCompileAPEX = false;
 
 		/// <summary>
 		/// Whether to include NvCloth.
 		/// </summary>
 		[RequiresUniqueBuildEnvironment]
+		[Obsolete("Deprecated in UE5.1 - No longer used as Chaos is always enabled.")]
 		public bool bCompileNvCloth = false;
 
 		/// <summary>
@@ -2435,11 +2441,13 @@ namespace UnrealBuildTool
 			get { return Inner.bUseVerse; }
 		}
 
+		[Obsolete("Deprecated in UE5.1 - No longer used as Chaos is always enabled.")]
 		public bool bCompileChaos
 		{
 			get { return Inner.bCompileChaos; }
 		}
 
+		[Obsolete("Deprecated in UE5.1 - No longer used as Chaos is always enabled.")]
 		public bool bUseChaos
 		{
 			get { return Inner.bUseChaos; }
@@ -2455,21 +2463,25 @@ namespace UnrealBuildTool
 			get { return Inner.bUseChaosChecked; }
 		}
 
+		[Obsolete("Deprecated in UE5.1 - No longer used in engine.")]
 		public bool bCustomSceneQueryStructure
 		{
 			get { return Inner.bCustomSceneQueryStructure; }
 		}
 
+		[Obsolete("Deprecated in UE5.1 - No longer used as Chaos is always enabled.")]
 		public bool bCompilePhysX
 		{
 			get { return Inner.bCompilePhysX; }
 		}
 
+		[Obsolete("Deprecated in UE5.1 - No longer used as Chaos is always enabled.")]
 		public bool bCompileAPEX
 		{
 			get { return Inner.bCompileAPEX; }
 		}
 
+		[Obsolete("Deprecated in UE5.1 - No longer used as Chaos is always enabled.")]
 		public bool bCompileNvCloth
 		{
 			get { return Inner.bCompileNvCloth; }
diff --git a/Engine/Source/Programs/UnrealBuildTool/Platform/Android/UEBuildAndroid.cs b/Engine/Source/Programs/UnrealBuildTool/Platform/Android/UEBuildAndroid.cs
--- a/Engine/Source/Programs/UnrealBuildTool/Platform/Android/UEBuildAndroid.cs
+++ b/Engine/Source/Programs/UnrealBuildTool/Platform/Android/UEBuildAndroid.cs
@@ -150,10 +150,6 @@
 
 		public override void ValidateTarget(TargetRules Target)
 		{
-			Target.bCompilePhysX = true;
-			Target.bCompileAPEX = false;
-			Target.bCompileNvCloth = false;
-
 			Target.bCompileRecast = true;
 			Target.bCompileISPC = false;
 		}
diff --git a/Engine/Source/Programs/UnrealBuildTool/Platform/IOS/UEBuildIOS.cs b/Engine/Source/Programs/UnrealBuildTool/Platform/IOS/UEBuildIOS.cs
--- a/Engine/Source/Programs/UnrealBuildTool/Platform/IOS/UEBuildIOS.cs
+++ b/Engine/Source/Programs/UnrealBuildTool/Platform/IOS/UEBuildIOS.cs
@@ -738,15 +738,6 @@
 
 		public override void ResetTarget(TargetRules Target)
 		{
-			// we currently don't have any simulator libs for PhysX
-			if (Target.Architecture == "-simulator")
-			{
-				Target.bCompilePhysX = false;
-			}
-
-			Target.bCompileAPEX = false;
-			Target.bCompileNvCloth = false;
-
 			Target.bDeployAfterCompile = true;
 
 			Target.IOSPlatform.ProjectSettings = ((IOSPlatform)GetBuildPlatform(Target.Platform)).ReadProjectSettings(Target.ProjectFile);
diff --git a/Engine/Source/Programs/UnrealBuildTool/Platform/Linux/UEBuildLinux.cs b/Engine/Source/Programs/UnrealBuildTool/Platform/Linux/UEBuildLinux.cs
--- a/Engine/Source/Programs/UnrealBuildTool/Platform/Linux/UEBuildLinux.cs
+++ b/Engine/Source/Programs/UnrealBuildTool/Platform/Linux/UEBuildLinux.cs
@@ -261,10 +261,6 @@ namespace UnrealBuildTool
 				throw new BuildException("LTO (LTCG) for modular builds is not supported (lld is not currently used for dynamic libraries).");
 			}
 
-			// depends on arch, APEX cannot be as of November'16 compiled for AArch32/64
-			Target.bCompileAPEX = Target.bCompileAPEX && Target.Architecture.StartsWith("x86_64");
-			Target.bCompileNvCloth = Target.bCompileNvCloth && Target.Architecture.StartsWith("x86_64");
-
 			if (Target.GlobalDefinitions.Contains("USE_NULL_RHI=1"))
 			{				
 				Target.bCompileCEF3 = false;
diff --git a/Engine/Source/Runtime/AnimGraphRuntime/AnimGraphRuntime.Build.cs b/Engine/Source/Runtime/AnimGraphRuntime/AnimGraphRuntime.Build.cs
--- a/Engine/Source/Runtime/AnimGraphRuntime/AnimGraphRuntime.Build.cs
+++ b/Engine/Source/Runtime/AnimGraphRuntime/AnimGraphRuntime.Build.cs
@@ -31,13 +31,10 @@ public class AnimGraphRuntime : ModuleRules
                 "Eigen"
                 );
 
-        if (Target.bCompileChaos || Target.bUseChaos)
-        {
-            PublicDependencyModuleNames.AddRange(
-                new string[] {
-					"GeometryCollectionEngine",
-                }
-            );
-        }
+        PublicDependencyModuleNames.AddRange(
+            new string[] {
+				"GeometryCollectionEngine",
+            }
+        );
     }
 }
diff --git a/Engine/Source/Runtime/AnimGraphRuntime/Private/BoneControllers/AnimNode_RigidBody.cpp b/Engine/Source/Runtime/AnimGraphRuntime/Private/BoneControllers/AnimNode_RigidBody.cpp
--- a/Engine/Source/Runtime/AnimGraphRuntime/Private/BoneControllers/AnimNode_RigidBody.cpp
+++ b/Engine/Source/Runtime/AnimGraphRuntime/Private/BoneControllers/AnimNode_RigidBody.cpp
@@ -505,18 +505,8 @@ void FAnimNode_RigidBody::RunPhysicsSimulation(float DeltaSeconds, const FVector
 
 	const int32 MaxSteps = RBAN_MaxSubSteps;
 	const float MaxDeltaSeconds = 1.f / 30.f;
-#if !WITH_CHAOS
-	const int32 NumSteps = FMath::Clamp(FMath::CeilToInt(DeltaSeconds / MaxDeltaSeconds), 1, MaxSteps);
-	const float StepDeltaTime = DeltaSeconds / float(NumSteps);
-	for (int32 Step = 1; Step <= NumSteps; Step++)
-	{
-		// We call the _AssumesLocked version here without a lock as the simulation is local to this node and we know
-		// we're not going to alter anything while this is running.
-		PhysicsSimulation->Simulate_AssumesLocked(StepDeltaTime, SimSpaceGravity);
-	}
-#else
+
 	PhysicsSimulation->Simulate_AssumesLocked(DeltaSeconds, MaxDeltaSeconds, MaxSteps, SimSpaceGravity);
-#endif
 }
 
 void FAnimNode_RigidBody::FlushDeferredSimulationTask()
@@ -813,7 +803,6 @@ void FAnimNode_RigidBody::EvaluateSkeletalControl_AnyThread(FComponentSpacePoseC
 			UpdateWorldForces(CompWorldSpaceTM, BaseBoneTM, DeltaSeconds);
 			SimSpaceGravity = WorldVectorToSpaceNoScale(SimulationSpace, WorldSpaceGravity, CompWorldSpaceTM, BaseBoneTM);
 
-#if WITH_CHAOS
 			FSimSpaceSettings* UseSimSpaceSettings = &SimSpaceSettings;
 			if (bRBAN_SimSpace_EnableOverride)
 			{
@@ -866,20 +855,17 @@ void FAnimNode_RigidBody::EvaluateSkeletalControl_AnyThread(FComponentSpacePoseC
 				SolverIterations.SolverPushOutIterations,
 				SolverIterations.JointPushOutIterations,
 				SolverIterations.CollisionPushOutIterations);
-#endif
 
 			if (!bUseDeferredSimulationTask)
 			{
 				RunPhysicsSimulation(DeltaSeconds, SimSpaceGravity);
 			}
 
-#if WITH_CHAOS
 			// Draw here even if the simulation is deferred since we want the shapes drawn relative to the current transform
 			if (bRBAN_DebugDraw)
 			{
 				PhysicsSimulation->DebugDraw();
 			}
-#endif
 		}
 		
 		//write back to animation system
@@ -1078,7 +1064,7 @@ void FAnimNode_RigidBody::InitPhysics(const UAnimInstance* InAnimInstance)
 		TArray<FConstraintInstance*> HighLevelConstraintInstances;
 
 		// Chaos relies on the initial pose to set up constraint positions
-		bool bCreateBodiesInRefPose = (WITH_CHAOS != 0);
+		constexpr bool bCreateBodiesInRefPose = true;
 		SkeletalMeshComp->InstantiatePhysicsAssetRefPose(
 			*UsePhysicsAsset, 
 			SimulationSpace == ESimulationSpace::WorldSpace ? SkeletalMeshComp->GetComponentToWorld().GetScale3D() : FVector(1.f), 
@@ -1148,9 +1134,7 @@ void FAnimNode_RigidBody::InitPhysics(const UAnimInstance* InAnimInstance)
 						IgnoreCollisionActors.Add(NewBodyHandle);
 					}
 
-#if WITH_CHAOS
 					NewBodyHandle->SetName(BodySetup->BoneName);
-#endif
 				}
 			}
 		}
@@ -1282,7 +1266,6 @@ void FAnimNode_RigidBody::InitPhysics(const UAnimInstance* InAnimInstance)
 
 		CollectClothColliderObjects(SkeletalMeshComp);
 
-#if WITH_CHAOS
 		SolverSettings = UsePhysicsAsset->SolverSettings;
 		PhysicsSimulation->SetSolverSettings(
 			SolverSettings.FixedTimeStep,
@@ -1300,7 +1283,6 @@ void FAnimNode_RigidBody::InitPhysics(const UAnimInstance* InAnimInstance)
 			SolverIterations.SolverPushOutIterations,
 			SolverIterations.JointPushOutIterations,
 			SolverIterations.CollisionPushOutIterations);
-#endif
 	}
 }
 
@@ -1484,7 +1466,7 @@ void FAnimNode_RigidBody::PreUpdate(const UAnimInstance* InAnimInstance)
 	APawn* PawnOwner = InAnimInstance->TryGetPawnOwner();
 	UPawnMovementComponent* MovementComp = PawnOwner ? PawnOwner->GetMovementComponent() : nullptr;
 
-#if WITH_EDITOR && WITH_CHAOS
+#if WITH_EDITOR
 	if (bEnableWorldGeometry && SimulationSpace != ESimulationSpace::WorldSpace && SKC && SKC->GetRelativeScale3D() != FVector(1.f, 1.f, 1.f))
 	{
 		FMessageLog("PIE").Warning(FText::Format(LOCTEXT("WorldCollisionComponentSpace", "Trying to use world collision without world space simulation for scaled ''{0}''. This is not supported, please change SimulationSpace to WorldSpace"),
@@ -1584,11 +1566,6 @@ void FAnimNode_RigidBody::CollectClothColliderObjects(const USkeletalMeshCompone
 {
 	if (bUseExternalClothCollision && bRBAN_IncludeClothColliders && SkeletalMeshComp && PhysicsSimulation)
 	{
-		// @todo(ccaulfield): is there an engine-independent way to do this?
-#if WITH_PHYSX && PHYSICS_INTERFACE_PHYSX
-		SCOPED_SCENE_READ_LOCK(PhysScene ? PhysScene->GetPxScene() : nullptr);
-#endif
-
 		const TArray<FClothCollisionSource>& SkeletalMeshClothCollisionSources = SkeletalMeshComp->GetClothCollisionSources();
 		
 		for (const FClothCollisionSource& ClothCollisionSource : SkeletalMeshClothCollisionSources)
@@ -1605,18 +1582,10 @@ void FAnimNode_RigidBody::CollectClothColliderObjects(const USkeletalMeshCompone
 				{
 					FBodyInstance* const BodyInstance = BodyInstances[BodyInstanceIndex];
 
-#if WITH_PHYSX && PHYSICS_INTERFACE_PHYSX
-					// Not sure why this happens, adding check to fix crash in CheckRBN engine test.
-					if (BodyInstance.BodySetup != nullptr)
-					{
-						ImmediatePhysics::FActorHandle* const ActorHandle = PhysicsSimulation->CreateActor(ImmediatePhysics::EActorType::StaticActor, BodyInstance, BodyInstance->GetUnrealWorldTransform());
-						ClothColliders.Add(FClothCollider(ActorHandle, SourceComponent, BodyInstance->InstanceBoneIndex));
-					}
-#elif WITH_CHAOS
 					ImmediatePhysics::FActorHandle* const ActorHandle = PhysicsSimulation->CreateActor(ImmediatePhysics::EActorType::KinematicActor, BodyInstance, BodyInstance->GetUnrealWorldTransform());
 					PhysicsSimulation->AddToCollidingPairs(ActorHandle); // <-allow collision between this actor and all dynamic actors.
 					ClothColliders.Add(FClothCollider(ActorHandle, SourceComponent, BodyInstance->InstanceBoneIndex));
-#endif
+
 					// Terminate the instance.
 					if (BodyInstance->IsValidBodyInstance())
 					{
@@ -1670,11 +1639,6 @@ void FAnimNode_RigidBody::CollectWorldObjects()
 		TArray<FOverlapResult> Overlaps;
 		UnsafeWorld->OverlapMultiByChannel(Overlaps, CachedBounds.Center, FQuat::Identity, OverlapChannel, FCollisionShape::MakeSphere(CachedBounds.W), QueryParams, FCollisionResponseParams(ECR_Overlap));
 
-		// @todo(ccaulfield): is there an engine-independent way to do this?
-#if WITH_PHYSX && PHYSICS_INTERFACE_PHYSX
-		SCOPED_SCENE_READ_LOCK(PhysScene ? PhysScene->GetPxScene() : nullptr); //TODO: expose this part to the anim node
-#endif
-
 		for (const FOverlapResult& Overlap : Overlaps)
 		{
 			if (UPrimitiveComponent* OverlapComp = Overlap.GetComponent())
@@ -1687,14 +1651,6 @@ void FAnimNode_RigidBody::CollectWorldObjects()
 				}
 				else
 				{
-#if WITH_PHYSX && PHYSICS_INTERFACE_PHYSX
-					ComponentsInSim.Add(OverlapComp);
-					// Not sure why this happens, adding check to fix crash in CheckRBN engine test.
-					if (OverlapComp->BodyInstance.BodySetup != nullptr)
-					{
-						PhysicsSimulation->CreateActor(ImmediatePhysics::EActorType::StaticActor, &OverlapComp->BodyInstance, OverlapComp->BodyInstance.GetUnrealWorldTransform());
-					}
-#elif WITH_CHAOS
 					// New object - add it to the sim
 					const bool bIsSelf = (UnsafeOwner == OverlapComp->GetOwner());
 					if (!bIsSelf)
@@ -1704,7 +1660,6 @@ void FAnimNode_RigidBody::CollectWorldObjects()
 						PhysicsSimulation->AddToCollidingPairs(ActorHandle);
 						ComponentsInSim.Add(OverlapComp, FWorldObject(ActorHandle, ComponentsInSimTick));
 					}
-#endif
 				}
 			}
 		}
@@ -1714,7 +1669,6 @@ void FAnimNode_RigidBody::CollectWorldObjects()
 // Flag invalid objects for purging
 void FAnimNode_RigidBody::ExpireWorldObjects()
 {
-#if WITH_CHAOS
 	// Invalidate deleted and expired world objects
 	TArray<const UPrimitiveComponent*> PrunedEntries;
 	for (auto& WorldEntry : ComponentsInSim)
@@ -1736,12 +1690,10 @@ void FAnimNode_RigidBody::ExpireWorldObjects()
 			WorldObject.bExpired = true;
 		}
 	}
-#endif
 }
 
 void FAnimNode_RigidBody::PurgeExpiredWorldObjects()
 {
-#if WITH_CHAOS
 	// Destroy expired simulated objects
 	TArray<const UPrimitiveComponent*> PurgedEntries;
 	for (auto& WorldEntry : ComponentsInSim)
@@ -1762,7 +1714,6 @@ void FAnimNode_RigidBody::PurgeExpiredWorldObjects()
 	{
 		ComponentsInSim.Remove(PurgedEntry);
 	}
-#endif
 }
 
 // Update the transforms of the world objects we added to the sim. This is required
@@ -1771,7 +1722,7 @@ void FAnimNode_RigidBody::PurgeExpiredWorldObjects()
 void FAnimNode_RigidBody::UpdateWorldObjects(const FTransform& SpaceTransform)
 {
 	LLM_SCOPE_BYNAME(TEXT("Animation/RigidBody")); 
-#if WITH_CHAOS
+
 	if (SimulationSpace != ESimulationSpace::WorldSpace)
 	{
 		for (const auto& WorldEntry : ComponentsInSim)
@@ -1793,7 +1744,6 @@ void FAnimNode_RigidBody::UpdateWorldObjects(const FTransform& SpaceTransform)
 			}
 		}
 	}
-#endif
 }
 
 void FAnimNode_RigidBody::InitializeBoneReferences(const FBoneContainer& RequiredBones) 
@@ -1897,11 +1847,7 @@ void FAnimNode_RigidBody::InitializeBoneReferences(const FBoneContainer& Require
 
 		if (PhysicsSimulation)
 		{
-#if WITH_CHAOS
 			PhysicsSimulation->SetNumActiveBodies(NumSimulatedBodies, SimulatedBodyIndices);
-#else
-			PhysicsSimulation->SetNumActiveBodies(NumSimulatedBodies);
-#endif
 		}
 
 		// We're switching to a new LOD, this invalidates our captured poses.
@@ -1912,7 +1858,6 @@ void FAnimNode_RigidBody::InitializeBoneReferences(const FBoneContainer& Require
 
 void FAnimNode_RigidBody::AddImpulseAtLocation(FVector Impulse, FVector Location, FName BoneName)
 {
-#if WITH_CHAOS
 	// Find the body. This is currently only used in the editor and will need optimizing if used in game
 	for (int32 BodyIndex = 0; BodyIndex < Bodies.Num(); ++BodyIndex)
 	{
@@ -1922,7 +1867,6 @@ void FAnimNode_RigidBody::AddImpulseAtLocation(FVector Impulse, FVector Location
 			Body->AddImpulseAtLocation(Impulse, Location);
 		}
 	}
-#endif
 }
 
 void FAnimNode_RigidBody::OnInitializeAnimInstance(const FAnimInstanceProxy* InProxy, const UAnimInstance* InAnimInstance)
diff --git a/Engine/Source/Runtime/ClothingSystemRuntimeCommon/Private/ClothingAsset.cpp b/Engine/Source/Runtime/ClothingSystemRuntimeCommon/Private/ClothingAsset.cpp
--- a/Engine/Source/Runtime/ClothingSystemRuntimeCommon/Private/ClothingAsset.cpp
+++ b/Engine/Source/Runtime/ClothingSystemRuntimeCommon/Private/ClothingAsset.cpp
@@ -979,14 +979,6 @@ void UClothingAssetCommon::CalculateReferenceBoneIndex()
 	// Starts at root
 	ReferenceBoneIndex = 0;
 
-#if WITH_APEX_CLOTHING
-	const TSubclassOf<class UClothingSimulationFactory> ClothingSimulationFactory = UClothingSimulationFactory::GetDefaultClothingSimulationFactoryClass();
-	if (ClothingSimulationFactory->GetName() == TEXT("ClothingSimulationFactoryNv"))
-	{
-		return;
-	}
-#endif
-
 	// Find the root bone for this clothing asset (common bone for all used bones)
 	typedef TArray<int32> BoneIndexArray;
 
diff --git a/Engine/Source/Runtime/ClothingSystemRuntimeInterface/Private/ClothingSimulationFactory.cpp b/Engine/Source/Runtime/ClothingSystemRuntimeInterface/Private/ClothingSimulationFactory.cpp
--- a/Engine/Source/Runtime/ClothingSystemRuntimeInterface/Private/ClothingSimulationFactory.cpp
+++ b/Engine/Source/Runtime/ClothingSystemRuntimeInterface/Private/ClothingSimulationFactory.cpp
@@ -10,21 +10,10 @@ namespace ClothingSimulationFactoryConsoleVariables
 {
 	TAutoConsoleVariable<FString> CVarDefaultClothingSimulationFactoryClass(
 		TEXT("p.Cloth.DefaultClothingSimulationFactoryClass"),
-#if WITH_CHAOS_CLOTHING
 		TEXT("ChaosClothingSimulationFactory"),  // Chaos is the default provider when Chaos Cloth is enabled
-#elif WITH_APEX_CLOTHING
-		TEXT("ClothingSimulationFactoryNv"),  // otherwise it's nv cloth
-#else
-		TEXT(""),  // otherwise it's none
-#endif
 		TEXT("The class name of the default clothing simulation factory.\n")
 		TEXT("Known providers are:\n")
-#if WITH_CHAOS_CLOTHING
 		TEXT("ChaosClothingSimulationFactory\n")
-#endif
-#if WITH_APEX_CLOTHING
-		TEXT("ClothingSimulationFactoryNv\n")
-#endif
 		, ECVF_Cheat);
 }
 
diff --git a/Engine/Source/Runtime/ClothingSystemRuntimeNv/Private/ClothingSimulationFactoryNv.cpp b/Engine/Source/Runtime/ClothingSystemRuntimeNv/Private/ClothingSimulationFactoryNv.cpp
--- a/Engine/Source/Runtime/ClothingSystemRuntimeNv/Private/ClothingSimulationFactoryNv.cpp
+++ b/Engine/Source/Runtime/ClothingSystemRuntimeNv/Private/ClothingSimulationFactoryNv.cpp
@@ -2,36 +2,21 @@
 
 #include "ClothingSimulationFactoryNv.h"
 
-#if WITH_NVCLOTH
-#include "UObject/Package.h"
-#include "ClothPhysicalMeshData.h"  // For EWeightMapTargetCommon
-#include "ClothingSimulationNv.h"
-#include "ClothingSimulationInteractorNv.h"
-#endif
-
+// Legacy support allowing the loading of older classes that may still reference old NvCloth simulation
+// classes. These are no longer expected to work but should still load so the references can be safely changed.
+PRAGMA_DISABLE_DEPRECATION_WARNINGS
 IClothingSimulation* UClothingSimulationFactoryNv::CreateSimulation()
 {
-#if WITH_NVCLOTH
-	return new FClothingSimulationNv();
-#else
 	return nullptr;
-#endif
 }
 
 void UClothingSimulationFactoryNv::DestroySimulation(IClothingSimulation* InSimulation)
 {
-#if WITH_NVCLOTH
-	delete InSimulation;
-#endif
 }
 
 bool UClothingSimulationFactoryNv::SupportsAsset(UClothingAssetBase* InAsset)
 {
-#if WITH_NVCLOTH
-	return true;
-#else
 	return false;
-#endif
 }
 
 bool UClothingSimulationFactoryNv::SupportsRuntimeInteraction()
@@ -41,31 +26,16 @@ bool UClothingSimulationFactoryNv::SupportsRuntimeInteraction()
 
 UClothingSimulationInteractor* UClothingSimulationFactoryNv::CreateInteractor()
 {
-#if WITH_NVCLOTH
-	return NewObject<UClothingSimulationInteractorNv>(GetTransientPackage());
-#else
 	return nullptr;
-#endif
 }
 
 TArrayView<const TSubclassOf<UClothConfigBase>> UClothingSimulationFactoryNv::GetClothConfigClasses() const
 {
-#if WITH_NVCLOTH
-	static const TArray<TSubclassOf<UClothConfigBase>> ClothConfigClasses(
-		{
-			TSubclassOf<UClothConfigBase>(UClothConfigNv::StaticClass())
-		});
-	return ClothConfigClasses;
-#else
 	return TArrayView<const TSubclassOf<UClothConfigBase>>();
-#endif
 }
 
 const UEnum* UClothingSimulationFactoryNv::GetWeightMapTargetEnum() const
 {
-#if WITH_NVCLOTH
-	return StaticEnum<EWeightMapTargetCommon>();
-#else
 	return nullptr;
-#endif
 }
+PRAGMA_ENABLE_DEPRECATION_WARNINGS
diff --git a/Engine/Source/Runtime/ClothingSystemRuntimeNv/Private/ClothingSimulationInteractorNv.cpp b/Engine/Source/Runtime/ClothingSystemRuntimeNv/Private/ClothingSimulationInteractorNv.cpp
--- a/Engine/Source/Runtime/ClothingSystemRuntimeNv/Private/ClothingSimulationInteractorNv.cpp
+++ b/Engine/Source/Runtime/ClothingSystemRuntimeNv/Private/ClothingSimulationInteractorNv.cpp
@@ -4,96 +4,32 @@
 
 #include "ClothingSimulationNv.h"
 
+// Stub implementations to allow older assets that reference this interactor to load safely so they can
+// be updated to use the Chaos version in future.
 void UClothingSimulationInteractorNv::PhysicsAssetUpdated()
 {
-#if WITH_NVCLOTH
-	Commands.Add(NvInteractorCommand::CreateLambda([](FClothingSimulationNv* InSimulation, FClothingSimulationContextNv* InContext)
-	{
-		InSimulation->RefreshPhysicsAsset();
-	}));
-#endif
 }
 
 void UClothingSimulationInteractorNv::ClothConfigUpdated()
 {
-#if WITH_NVCLOTH
-	Commands.Add(NvInteractorCommand::CreateLambda([](FClothingSimulationNv* InSimulation, FClothingSimulationContextNv* InContext)
-	{
-		InSimulation->RefreshClothConfig();
-	}));
-#endif
 }
 
 void UClothingSimulationInteractorNv::Sync(IClothingSimulation* InSimulation, IClothingSimulationContext* InContext)
 {
-#if WITH_NVCLOTH
-	check(InSimulation);
-	check(InContext);
-
-	FClothingSimulationNv* NvSim = static_cast<FClothingSimulationNv*>(InSimulation);
-	FClothingSimulationContextNv* NvContext = static_cast<FClothingSimulationContextNv*>(InContext);
-
-	for(NvInteractorCommand& Command : Commands)
-	{
-		Command.Execute(NvSim, NvContext);
-	}
-	Commands.Reset();
-
-	// Call base class' sync 
-	UClothingSimulationInteractor::Sync(InSimulation, InContext);
-#endif
 }
 
 void UClothingSimulationInteractorNv::SetAnimDriveSpringStiffness(float InStiffness)
 {
-#if WITH_NVCLOTH
-	Commands.Add(NvInteractorCommand::CreateLambda([InStiffness](FClothingSimulationNv* InSimulation, FClothingSimulationContextNv* InContext)
-	{
-		InSimulation->ExecutePerActor([InStiffness](FClothingActorNv& Actor)
-		{
-			Actor.CurrentAnimDriveSpringStiffness = InStiffness;
-		});
-	}));
-#endif
 }
 
 void UClothingSimulationInteractorNv::SetAnimDriveDamperStiffness(float InStiffness)
 {
-#if WITH_NVCLOTH
-	Commands.Add(NvInteractorCommand::CreateLambda([InStiffness](FClothingSimulationNv* InSimulation, FClothingSimulationContextNv* InContext)
-	{
-		InSimulation->ExecutePerActor([InStiffness](FClothingActorNv& Actor)
-		{
-			Actor.CurrentAnimDriveDamperStiffness = InStiffness;
-		});
-	}));
-#endif
 }
 void UClothingSimulationInteractorNv::EnableGravityOverride(const FVector& InVector)
 {
-#if WITH_NVCLOTH
-	Commands.Add(NvInteractorCommand::CreateLambda([InVector](FClothingSimulationNv* InSimulation, FClothingSimulationContextNv* InContext)
-	{
-		InSimulation->ExecutePerActor([InVector](FClothingActorNv& Actor)
-		{
-			Actor.bUseGravityOverride = true;
-			Actor.GravityOverride = InVector;
-		});
-	}));
-#endif
 }
 
 void UClothingSimulationInteractorNv::DisableGravityOverride()
 {
-#if WITH_NVCLOTH
-	Commands.Add(NvInteractorCommand::CreateLambda([](FClothingSimulationNv* InSimulation, FClothingSimulationContextNv* InContext)
-	{
-		InSimulation->ExecutePerActor([](FClothingActorNv& Actor)
-		{
-			Actor.bUseGravityOverride = false;
-			Actor.GravityOverride = FVector(0.0f, 0.0f, 0.0f);
-		});
-	}));
-#endif
 }
 
