Added looping flag to anim sequences and blend spaces. This is used to set the default loop value when creating asset player nodes

diff --git a/Engine/Source/Editor/AnimGraph/Private/AnimGraphNode_BlendSpaceEvaluator.cpp b/Engine/Source/Editor/AnimGraph/Private/AnimGraphNode_BlendSpaceEvaluator.cpp
--- a/Engine/Source/Editor/AnimGraph/Private/AnimGraphNode_BlendSpaceEvaluator.cpp
+++ b/Engine/Source/Editor/AnimGraph/Private/AnimGraphNode_BlendSpaceEvaluator.cpp
@@ -95,6 +95,14 @@ void UAnimGraphNode_BlendSpaceEvaluator::SetAnimationAsset(UAnimationAsset* Asse
 	}
 }
 
+void UAnimGraphNode_BlendSpaceEvaluator::CopySettingsFromAnimationAsset(UAnimationAsset* Asset)
+{
+	if (UBlendSpace* BlendSpace = Cast<UBlendSpace>(Asset))
+	{
+		Node.SetLoop(BlendSpace->bLoop);
+	}
+}
+
 void UAnimGraphNode_BlendSpaceEvaluator::OnOverrideAssets(IAnimBlueprintNodeOverrideAssetsContext& InContext) const
 {
 	if(InContext.GetAssets().Num() > 0)
diff --git a/Engine/Source/Editor/AnimGraph/Private/AnimGraphNode_BlendSpacePlayer.cpp b/Engine/Source/Editor/AnimGraph/Private/AnimGraphNode_BlendSpacePlayer.cpp
--- a/Engine/Source/Editor/AnimGraph/Private/AnimGraphNode_BlendSpacePlayer.cpp
+++ b/Engine/Source/Editor/AnimGraph/Private/AnimGraphNode_BlendSpacePlayer.cpp
@@ -109,6 +109,14 @@ void UAnimGraphNode_BlendSpacePlayer::SetAnimationAsset(UAnimationAsset* Asset)
 	}
 }
 
+void UAnimGraphNode_BlendSpacePlayer::CopySettingsFromAnimationAsset(UAnimationAsset* Asset)
+{
+	if (UBlendSpace* BlendSpace = Cast<UBlendSpace>(Asset))
+	{
+		Node.SetLoop(BlendSpace->bLoop);
+	}
+}
+
 void UAnimGraphNode_BlendSpacePlayer::OnOverrideAssets(IAnimBlueprintNodeOverrideAssetsContext& InContext) const
 {
 	if(InContext.GetAssets().Num() > 0)
diff --git a/Engine/Source/Editor/AnimGraph/Private/AnimGraphNode_SequenceEvaluator.cpp b/Engine/Source/Editor/AnimGraph/Private/AnimGraphNode_SequenceEvaluator.cpp
--- a/Engine/Source/Editor/AnimGraph/Private/AnimGraphNode_SequenceEvaluator.cpp
+++ b/Engine/Source/Editor/AnimGraph/Private/AnimGraphNode_SequenceEvaluator.cpp
@@ -128,6 +128,14 @@ void UAnimGraphNode_SequenceEvaluator::SetAnimationAsset(UAnimationAsset* Asset)
 	}
 }
 
+void UAnimGraphNode_SequenceEvaluator::CopySettingsFromAnimationAsset(UAnimationAsset* Asset)
+{
+	if (UAnimSequenceBase* Seq = Cast<UAnimSequence>(Asset))
+	{
+		Node.SetShouldLoop(Seq->bLoop);
+	}
+}
+
 void UAnimGraphNode_SequenceEvaluator::OnOverrideAssets(IAnimBlueprintNodeOverrideAssetsContext& InContext) const
 {
 	if(InContext.GetAssets().Num() > 0)
diff --git a/Engine/Source/Editor/AnimGraph/Private/AnimGraphNode_SequencePlayer.cpp b/Engine/Source/Editor/AnimGraph/Private/AnimGraphNode_SequencePlayer.cpp
--- a/Engine/Source/Editor/AnimGraph/Private/AnimGraphNode_SequencePlayer.cpp
+++ b/Engine/Source/Editor/AnimGraph/Private/AnimGraphNode_SequencePlayer.cpp
@@ -186,6 +186,14 @@ void UAnimGraphNode_SequencePlayer::SetAnimationAsset(UAnimationAsset* Asset)
 	}
 }
 
+void UAnimGraphNode_SequencePlayer::CopySettingsFromAnimationAsset(UAnimationAsset* Asset)
+{
+	if (UAnimSequenceBase* Seq = Cast<UAnimSequenceBase>(Asset))
+	{
+		Node.SetLoopAnimation(Seq->bLoop);
+	}
+}
+
 void UAnimGraphNode_SequencePlayer::OnOverrideAssets(IAnimBlueprintNodeOverrideAssetsContext& InContext) const
 {
 	if(InContext.GetAssets().Num() > 0)
diff --git a/Engine/Source/Editor/AnimGraph/Private/AnimationGraphSchema.cpp b/Engine/Source/Editor/AnimGraph/Private/AnimationGraphSchema.cpp
--- a/Engine/Source/Editor/AnimGraph/Private/AnimationGraphSchema.cpp
+++ b/Engine/Source/Editor/AnimGraph/Private/AnimationGraphSchema.cpp
@@ -579,6 +579,7 @@ void UAnimationGraphSchema::SpawnNodeFromAsset(UAnimationAsset* Asset, const FVe
 
 			UAnimGraphNode_AssetPlayerBase* NewNode = NewObject<UAnimGraphNode_AssetPlayerBase>(GetTransientPackage(), NewNodeClass);
 			NewNode->SetAnimationAsset(Asset);
+			NewNode->CopySettingsFromAnimationAsset(Asset);
 			Action.NodeTemplate = NewNode;
 
 			Action.PerformAction(Graph, PinIfAvailable, GraphPosition);
diff --git a/Engine/Source/Editor/AnimGraph/Public/AnimGraphNode_AssetPlayerBase.h b/Engine/Source/Editor/AnimGraph/Public/AnimGraphNode_AssetPlayerBase.h
--- a/Engine/Source/Editor/AnimGraph/Public/AnimGraphNode_AssetPlayerBase.h
+++ b/Engine/Source/Editor/AnimGraph/Public/AnimGraphNode_AssetPlayerBase.h
@@ -43,6 +43,9 @@ public:
 	virtual void ValidateAnimNodeDuringCompilation(USkeleton* ForSkeleton, FCompilerResultsLog& MessageLog) override;
 	virtual void SetAnimationAsset(UAnimationAsset* Asset) { check(false); /*Base function called*/ }
 
+	// Override this to copy any relevant settings from the animation asset used when creating a node with a drag/drop operation
+	virtual void CopySettingsFromAnimationAsset(UAnimationAsset* Asset) {}
+
 	// Helper function to gather menu actions from specific asset types supported by this node
 	static void GetMenuActionsHelper(
 		FBlueprintActionDatabaseRegistrar& InActionRegistrar,
diff --git a/Engine/Source/Editor/AnimGraph/Public/AnimGraphNode_BlendSpaceEvaluator.h b/Engine/Source/Editor/AnimGraph/Public/AnimGraphNode_BlendSpaceEvaluator.h
--- a/Engine/Source/Editor/AnimGraph/Public/AnimGraphNode_BlendSpaceEvaluator.h
+++ b/Engine/Source/Editor/AnimGraph/Public/AnimGraphNode_BlendSpaceEvaluator.h
@@ -42,6 +42,7 @@ class UAnimGraphNode_BlendSpaceEvaluator : public UAnimGraphNode_BlendSpaceBase
 
 	// UAnimGraphNode_AssetPlayerBase interface
 	virtual void SetAnimationAsset(UAnimationAsset* Asset) override;
+	virtual void CopySettingsFromAnimationAsset(UAnimationAsset* Asset) override;
 	// End of UAnimGraphNode_AssetPlayerBase interface
 
 private:
@@ -50,4 +51,5 @@ private:
 private:
 	/** Constructing FText strings can be costly, so we cache the node's title */
 	FNodeTitleTextTable CachedNodeTitles;
+
 };
diff --git a/Engine/Source/Editor/AnimGraph/Public/AnimGraphNode_BlendSpacePlayer.h b/Engine/Source/Editor/AnimGraph/Public/AnimGraphNode_BlendSpacePlayer.h
--- a/Engine/Source/Editor/AnimGraph/Public/AnimGraphNode_BlendSpacePlayer.h
+++ b/Engine/Source/Editor/AnimGraph/Public/AnimGraphNode_BlendSpacePlayer.h
@@ -45,5 +45,6 @@ class UAnimGraphNode_BlendSpacePlayer : public UAnimGraphNode_BlendSpaceBase
 
 	// UAnimGraphNode_AssetPlayerBase interface
 	virtual void SetAnimationAsset(UAnimationAsset* Asset) override;
+	virtual void CopySettingsFromAnimationAsset(UAnimationAsset* Asset) override;
 	// End of UAnimGraphNode_AssetPlayerBase interface
 };
diff --git a/Engine/Source/Editor/AnimGraph/Public/AnimGraphNode_SequenceEvaluator.h b/Engine/Source/Editor/AnimGraph/Public/AnimGraphNode_SequenceEvaluator.h
--- a/Engine/Source/Editor/AnimGraph/Public/AnimGraphNode_SequenceEvaluator.h
+++ b/Engine/Source/Editor/AnimGraph/Public/AnimGraphNode_SequenceEvaluator.h
@@ -48,5 +48,6 @@ class UAnimGraphNode_SequenceEvaluator : public UAnimGraphNode_AssetPlayerBase
 
 	// UAnimGraphNode_AssetPlayerBase interface
 	virtual void SetAnimationAsset(UAnimationAsset* Asset) override;
+	virtual void CopySettingsFromAnimationAsset(UAnimationAsset* Asset) override;
 	// End of UAnimGraphNode_AssetPlayerBase interface
 };
diff --git a/Engine/Source/Editor/AnimGraph/Public/AnimGraphNode_SequencePlayer.h b/Engine/Source/Editor/AnimGraph/Public/AnimGraphNode_SequencePlayer.h
--- a/Engine/Source/Editor/AnimGraph/Public/AnimGraphNode_SequencePlayer.h
+++ b/Engine/Source/Editor/AnimGraph/Public/AnimGraphNode_SequencePlayer.h
@@ -55,5 +55,6 @@ class UAnimGraphNode_SequencePlayer : public UAnimGraphNode_AssetPlayerBase
 
 	// UAnimGraphNode_AssetPlayerBase interface
 	virtual void SetAnimationAsset(UAnimationAsset* Asset) override;
+	virtual void CopySettingsFromAnimationAsset(UAnimationAsset* Asset) override;
 	// End of UAnimGraphNode_AssetPlayerBase interface
 };
diff --git a/Engine/Source/Runtime/AnimGraphRuntime/Private/AnimNodes/AnimNode_SequenceEvaluator.cpp b/Engine/Source/Runtime/AnimGraphRuntime/Private/AnimNodes/AnimNode_SequenceEvaluator.cpp
--- a/Engine/Source/Runtime/AnimGraphRuntime/Private/AnimNodes/AnimNode_SequenceEvaluator.cpp
+++ b/Engine/Source/Runtime/AnimGraphRuntime/Private/AnimNodes/AnimNode_SequenceEvaluator.cpp
@@ -169,6 +169,21 @@ bool FAnimNode_SequenceEvaluator::SetExplicitTime(float InTime)
 	return false;
 }
 
+bool FAnimNode_SequenceEvaluator::SetShouldLoop(bool bInShouldLoop)
+{
+#if WITH_EDITORONLY_DATA
+	bShouldLoop = bInShouldLoop;
+#endif
+
+	if (bool* bShouldLoopPtr = GET_INSTANCE_ANIM_NODE_DATA_PTR(bool, bShouldLoop))
+	{
+		*bShouldLoopPtr = bInShouldLoop;
+		return true;
+	}
+
+	return false;
+}
+
 bool FAnimNode_SequenceEvaluator::GetShouldLoop() const
 {
 	return GET_ANIM_NODE_DATA(bool, bShouldLoop);
diff --git a/Engine/Source/Runtime/AnimGraphRuntime/Public/AnimNodes/AnimNode_SequenceEvaluator.h b/Engine/Source/Runtime/AnimGraphRuntime/Public/AnimNodes/AnimNode_SequenceEvaluator.h
--- a/Engine/Source/Runtime/AnimGraphRuntime/Public/AnimNodes/AnimNode_SequenceEvaluator.h
+++ b/Engine/Source/Runtime/AnimGraphRuntime/Public/AnimNodes/AnimNode_SequenceEvaluator.h
@@ -79,6 +79,9 @@ public:
 	/** This only works if bTeleportToExplicitTime is false OR this node is set to use SyncGroup */
 	virtual bool GetShouldLoop() const { return true; }
 
+	// Set the animation to continue looping when it reaches the end
+	virtual bool SetShouldLoop(bool bInShouldLoop) { return false; }
+
 	/** If true, teleport to explicit time, does NOT advance time (does not trigger notifies, does not extract Root Motion, etc.)
 	If false, will advance time (will trigger notifies, extract root motion if applicable, etc.)
 	Note: using a sync group forces advancing time regardless of what this option is set to. */
@@ -151,6 +154,7 @@ public:
 	virtual float GetExplicitTime() const override;
 	virtual bool SetExplicitTime(float InTime) override;
 	virtual bool GetShouldLoop() const override;
+	virtual bool SetShouldLoop(bool bInShouldLoop) override;
 	virtual bool GetTeleportToExplicitTime() const override;
 	virtual TEnumAsByte<ESequenceEvalReinit::Type> GetReinitializationBehavior() const override;
 	virtual float GetStartPosition() const override;
diff --git a/Engine/Source/Runtime/Engine/Classes/Animation/AnimSequenceBase.h b/Engine/Source/Runtime/Engine/Classes/Animation/AnimSequenceBase.h
--- a/Engine/Source/Runtime/Engine/Classes/Animation/AnimSequenceBase.h
+++ b/Engine/Source/Runtime/Engine/Classes/Animation/AnimSequenceBase.h
@@ -53,6 +53,13 @@ public:
 	UPROPERTY(EditAnywhere, Category=Animation)
 	float RateScale;
 	
+	/** 
+	 * The default looping behavior of this animation.
+	 * Asset players can override this
+	 */
+	UPROPERTY(EditAnywhere, Category=Animation)
+	bool bLoop;
+
 	/**
 	 * Raw uncompressed float curve data 
 	 */
diff --git a/Engine/Source/Runtime/Engine/Classes/Animation/BlendSpace.h b/Engine/Source/Runtime/Engine/Classes/Animation/BlendSpace.h
--- a/Engine/Source/Runtime/Engine/Classes/Animation/BlendSpace.h
+++ b/Engine/Source/Runtime/Engine/Classes/Animation/BlendSpace.h
@@ -732,6 +732,13 @@ public:
 	UPROPERTY(EditAnywhere, Category = SampleSmoothing, meta = (DisplayName = "Smoothing"))
 	bool bTargetWeightInterpolationEaseInOut = true;
 
+	/** 
+	* The default looping behavior of this blend space.
+	* Asset players can override this
+	*/
+	UPROPERTY(EditAnywhere, Category=Animation)
+	bool bLoop = true;
+
 #if WITH_EDITORONLY_DATA
 	/** Preview Base pose for additive BlendSpace **/
 	UPROPERTY(EditAnywhere, Category = AdditiveSettings)
diff --git a/Engine/Source/Runtime/Engine/Private/Animation/AnimSequenceBase.cpp b/Engine/Source/Runtime/Engine/Private/Animation/AnimSequenceBase.cpp
--- a/Engine/Source/Runtime/Engine/Private/Animation/AnimSequenceBase.cpp
+++ b/Engine/Source/Runtime/Engine/Private/Animation/AnimSequenceBase.cpp
@@ -32,6 +32,7 @@ CSV_DECLARE_CATEGORY_MODULE_EXTERN(ENGINE_API, Animation);
 UAnimSequenceBase::UAnimSequenceBase(const FObjectInitializer& ObjectInitializer)
 	: Super(ObjectInitializer)
 	, RateScale(1.0f)
+	, bLoop(true)
 #if WITH_EDITORONLY_DATA
 	, DataModel(nullptr)
 	, bPopulatingDataModel(false)
