Fix for SequenceEvaluator nodes not showing up in blend weights view

diff --git a/Engine/Source/Runtime/AnimGraphRuntime/Private/AnimNodes/AnimNode_SequenceEvaluator.cpp b/Engine/Source/Runtime/AnimGraphRuntime/Private/AnimNodes/AnimNode_SequenceEvaluator.cpp
--- a/Engine/Source/Runtime/AnimGraphRuntime/Private/AnimNodes/AnimNode_SequenceEvaluator.cpp
+++ b/Engine/Source/Runtime/AnimGraphRuntime/Private/AnimNodes/AnimNode_SequenceEvaluator.cpp
@@ -70,6 +70,7 @@ void FAnimNode_SequenceEvaluatorBase::UpdateAssetPlayer(const FAnimationUpdateCo
 		else
 		{
 			InternalTimeAccumulator = CurrentExplicitTime;
+			CreateTickRecordForNode(Context, CurrentSequence, GetShouldLoop(), 0);
 		}
 	}
 
