GitHub #8712 : [AnimNotify] Implement GetNotifyName for PlayMontageNotify

diff --git a/Engine/Source/Runtime/AnimGraphRuntime/Private/AnimNotify_PlayMontageNotify.cpp b/Engine/Source/Runtime/AnimGraphRuntime/Private/AnimNotify_PlayMontageNotify.cpp
--- a/Engine/Source/Runtime/AnimGraphRuntime/Private/AnimNotify_PlayMontageNotify.cpp
+++ b/Engine/Source/Runtime/AnimGraphRuntime/Private/AnimNotify_PlayMontageNotify.cpp
@@ -27,6 +27,11 @@ void UAnimNotify_PlayMontageNotify::BranchingPointNotify(FBranchingPointNotifyPa
 	}
 }
 
+FString UAnimNotify_PlayMontageNotify::GetNotifyName_Implementation() const
+{
+	return NotifyName.ToString();
+}
+
 #if WITH_EDITOR
 bool UAnimNotify_PlayMontageNotify::CanBePlaced(UAnimSequenceBase* Animation) const
 {
@@ -73,6 +78,11 @@ void UAnimNotify_PlayMontageNotifyWindow::BranchingPointNotifyEnd(FBranchingPoin
 	}
 }
 
+FString UAnimNotify_PlayMontageNotifyWindow::GetNotifyName_Implementation() const
+{
+	return NotifyName.ToString();
+}
+
 #if WITH_EDITOR
 bool UAnimNotify_PlayMontageNotifyWindow::CanBePlaced(UAnimSequenceBase* Animation) const
 {
diff --git a/Engine/Source/Runtime/AnimGraphRuntime/Public/AnimNotifies/AnimNotify_PlayMontageNotify.h b/Engine/Source/Runtime/AnimGraphRuntime/Public/AnimNotifies/AnimNotify_PlayMontageNotify.h
--- a/Engine/Source/Runtime/AnimGraphRuntime/Public/AnimNotifies/AnimNotify_PlayMontageNotify.h
+++ b/Engine/Source/Runtime/AnimGraphRuntime/Public/AnimNotifies/AnimNotify_PlayMontageNotify.h
@@ -17,6 +17,8 @@ class ANIMGRAPHRUNTIME_API UAnimNotify_PlayMontageNotify : public UAnimNotify
 public:
 
 	virtual void BranchingPointNotify(FBranchingPointNotifyPayload& BranchingPointPayload) override;
+	
+	virtual FString GetNotifyName_Implementation() const override;
 
 #if WITH_EDITOR
 	virtual bool CanBePlaced(UAnimSequenceBase* Animation) const override;
@@ -42,6 +44,8 @@ public:
 	virtual void BranchingPointNotifyBegin(FBranchingPointNotifyPayload& BranchingPointPayload) override;
 	virtual void BranchingPointNotifyEnd(FBranchingPointNotifyPayload& BranchingPointPayload) override;
 
+	virtual FString GetNotifyName_Implementation() const override;
+
 #if WITH_EDITOR
 	virtual bool CanBePlaced(UAnimSequenceBase* Animation) const override;
 #endif
