Fix RBAN memory stomp due to out of bounds array access

diff --git a/Engine/Source/Runtime/AnimGraphRuntime/Private/BoneControllers/AnimNode_RigidBody.cpp b/Engine/Source/Runtime/AnimGraphRuntime/Private/BoneControllers/AnimNode_RigidBody.cpp
--- a/Engine/Source/Runtime/AnimGraphRuntime/Private/BoneControllers/AnimNode_RigidBody.cpp
+++ b/Engine/Source/Runtime/AnimGraphRuntime/Private/BoneControllers/AnimNode_RigidBody.cpp
@@ -1124,7 +1124,10 @@ void FAnimNode_RigidBody::InitPhysics(const UAnimInstance* InAnimInstance)
 					}
 					const int32 BodyIndex = Bodies.Add(NewBodyHandle);
 					const int32 SkeletonBoneIndex = MeshToSkeletonBoneIndex[InsertBone];
-					SkeletonBoneIndexToBodyIndex[SkeletonBoneIndex] = BodyIndex;
+					if (ensure(SkeletonBoneIndex >= 0))
+					{
+						SkeletonBoneIndexToBodyIndex[SkeletonBoneIndex] = BodyIndex;
+					}
 					BodyAnimData[BodyIndex].bIsSimulated = bSimulated;
 					NamesToHandles.Add(BodySetup->BoneName, NewBodyHandle);
 					BodyIndexToActorHandle[BodyInstance->InstanceBodyIndex] = NewBodyHandle;
