AnimNode_Rigidbody MasterAlpha rename. Part of UE-158645

diff --git a/Engine/Config/BaseEngine.ini b/Engine/Config/BaseEngine.ini
--- a/Engine/Config/BaseEngine.ini
+++ b/Engine/Config/BaseEngine.ini
@@ -1096,6 +1096,9 @@ DefaultWorldPartitionHLODLayer=/Engine/Maps/Templates/HLODs/HLODLayer_Instanced.
 
 +PropertyRedirects=(OldName="RecastNavMeshTileGenerationDebug.bHeightfieldSolidPostRadiusFiltering", NewName="bHeightfieldSolidPostInclusionBoundsFiltering")
 
+; Inclusive Terminology Pass
++PropertyRedirects=(OldName="SimSpaceSettings.MasterAlpha", NewName="SimSpaceSettings.WorldAlpha")
+
 [CoreUObject.Metadata]
 ; Note: UnrealHeaderTool should be rerun after making changes to MetadataRedirects to catch any keys specified in class headers
 +MetadataRedirects=(OldKey="K2Protected", NewKey="BlueprintProtected")
diff --git a/Engine/Source/Runtime/AnimGraphRuntime/Private/BoneControllers/AnimNode_RigidBody.cpp b/Engine/Source/Runtime/AnimGraphRuntime/Private/BoneControllers/AnimNode_RigidBody.cpp
--- a/Engine/Source/Runtime/AnimGraphRuntime/Private/BoneControllers/AnimNode_RigidBody.cpp
+++ b/Engine/Source/Runtime/AnimGraphRuntime/Private/BoneControllers/AnimNode_RigidBody.cpp
@@ -58,7 +58,7 @@ FAutoConsoleVariableRef CVarRigidBodyNodeIncludeClothColliders(TEXT("p.RigidBody
 bool bRBAN_SimSpace_EnableOverride = false;
 FSimSpaceSettings RBAN_SimSpaceOverride;
 FAutoConsoleVariableRef CVarRigidBodyNodeSpaceOverride(TEXT("p.RigidBodyNode.Space.Override"), bRBAN_SimSpace_EnableOverride, TEXT("Force-enable the advanced simulation space movement forces"), ECVF_Default);
-FAutoConsoleVariableRef CVarRigidBodyNodeSpaceMasterAlpha(TEXT("p.RigidBodyNode.Space.MasterAlpha"), RBAN_SimSpaceOverride.MasterAlpha, TEXT("RBAN SimSpaceSettings overrides"), ECVF_Default);
+FAutoConsoleVariableRef CVarRigidBodyNodeSpaceWorldAlpha(TEXT("p.RigidBodyNode.Space.WorldAlpha"), RBAN_SimSpaceOverride.WorldAlpha, TEXT("RBAN SimSpaceSettings overrides"), ECVF_Default);
 FAutoConsoleVariableRef CVarRigidBodyNodeSpaceVelScaleZ(TEXT("p.RigidBodyNode.Space.VelocityScaleZ"), RBAN_SimSpaceOverride.VelocityScaleZ, TEXT("RBAN SimSpaceSettings overrides"), ECVF_Default);
 FAutoConsoleVariableRef CVarRigidBodyNodeSpaceMaxCompLinVel(TEXT("p.RigidBodyNode.Space.MaxLinearVelocity"), RBAN_SimSpaceOverride.MaxLinearVelocity, TEXT("RBAN SimSpaceSettings overrides"), ECVF_Default);
 FAutoConsoleVariableRef CVarRigidBodyNodeSpaceMaxCompAngVel(TEXT("p.RigidBodyNode.Space.MaxAngularVelocity"), RBAN_SimSpaceOverride.MaxAngularVelocity, TEXT("RBAN SimSpaceSettings overrides"), ECVF_Default);
@@ -103,7 +103,7 @@ FAutoConsoleVariableRef CVarRigidBodyNodeSimulationTaskPriority(
 );
 
 FSimSpaceSettings::FSimSpaceSettings()
-	: MasterAlpha(0)
+	: WorldAlpha(0)
 	, VelocityScaleZ(1)
 	, MaxLinearVelocity(10000)
 	, MaxAngularVelocity(10000)
@@ -416,7 +416,7 @@ void FAnimNode_RigidBody::CalculateSimulationSpace(
 	SpaceAngularAcc = FVector::ZeroVector;
 
 	// If the system is disabled, nothing else to do
-	if ((Settings.MasterAlpha == 0.0f) || (Dt < SMALL_NUMBER))
+	if ((Settings.WorldAlpha == 0.0f) || (Dt < SMALL_NUMBER))
 	{
 		return;
 	}
@@ -837,7 +837,7 @@ void FAnimNode_RigidBody::EvaluateSkeletalControl_AnyThread(FComponentSpacePoseC
 				SimulationAngularAcceleration);
 
 			PhysicsSimulation->SetSimulationSpaceSettings(
-				UseSimSpaceSettings->MasterAlpha, 
+				UseSimSpaceSettings->WorldAlpha,
 				UseSimSpaceSettings->ExternalLinearDragV);
 
 			PhysicsSimulation->SetSolverSettings(
diff --git a/Engine/Source/Runtime/AnimGraphRuntime/Public/BoneControllers/AnimNode_RigidBody.h b/Engine/Source/Runtime/AnimGraphRuntime/Public/BoneControllers/AnimNode_RigidBody.h
--- a/Engine/Source/Runtime/AnimGraphRuntime/Public/BoneControllers/AnimNode_RigidBody.h
+++ b/Engine/Source/Runtime/AnimGraphRuntime/Public/BoneControllers/AnimNode_RigidBody.h
@@ -52,11 +52,11 @@ struct ANIMGRAPHRUNTIME_API FSimSpaceSettings
 
 	FSimSpaceSettings();
 
-	// Global multipler on the effects of simulation space movement. Must be in range [0, 1]. If MasterAlpha = 0.0, the system is disabled and the simulation will
-	// be fully local (i.e., world-space actor movement and rotation does not affect the simulation). When MasterAlpha = 1.0 the simulation effectively acts as a 
+	// Global multipler on the effects of simulation space movement. Must be in range [0, 1]. If WorldAlpha = 0.0, the system is disabled and the simulation will
+	// be fully local (i.e., world-space actor movement and rotation does not affect the simulation). When WorldAlpha = 1.0 the simulation effectively acts as a 
 	// world-space sim, but with the ability to apply limits using the other parameters.
 	UPROPERTY(EditAnywhere, BlueprintReadWrite, Category = Settings, meta = (ClampMin = "0.0", ClampMax = "1.0"))
-	float MasterAlpha;
+	float WorldAlpha;
 
 	// Multiplier on the Z-component of velocity and acceleration that is passed to the simulation. Usually from 0.0 to 1.0 to 
 	// reduce the effects of jumping and crouching on the simulation, but it can be higher than 1.0 if you need to exaggerate this motion for some reason.
