AnimNode_Rigidbody MasterAlpha rename. Part of UE-158645
Adding the old property as deprecated, so the c++ continues compiling for external users of the engine

diff --git a/Engine/Source/Runtime/AnimGraphRuntime/Public/BoneControllers/AnimNode_RigidBody.h b/Engine/Source/Runtime/AnimGraphRuntime/Public/BoneControllers/AnimNode_RigidBody.h
--- a/Engine/Source/Runtime/AnimGraphRuntime/Public/BoneControllers/AnimNode_RigidBody.h
+++ b/Engine/Source/Runtime/AnimGraphRuntime/Public/BoneControllers/AnimNode_RigidBody.h
@@ -52,12 +52,25 @@ struct ANIMGRAPHRUNTIME_API FSimSpaceSettings
 
 	FSimSpaceSettings();
 
+	// Disable deprecation errors by providing defaults wrapped with pragma disable
+	PRAGMA_DISABLE_DEPRECATION_WARNINGS
+	~FSimSpaceSettings() = default;
+	FSimSpaceSettings(FSimSpaceSettings const&) = default;
+	FSimSpaceSettings& operator=(const FSimSpaceSettings &) = default;
+	PRAGMA_ENABLE_DEPRECATION_WARNINGS
+
 	// Global multipler on the effects of simulation space movement. Must be in range [0, 1]. If WorldAlpha = 0.0, the system is disabled and the simulation will
 	// be fully local (i.e., world-space actor movement and rotation does not affect the simulation). When WorldAlpha = 1.0 the simulation effectively acts as a 
 	// world-space sim, but with the ability to apply limits using the other parameters.
 	UPROPERTY(EditAnywhere, BlueprintReadWrite, Category = Settings, meta = (ClampMin = "0.0", ClampMax = "1.0"))
 	float WorldAlpha;
 
+#if WITH_EDITORONLY_DATA
+	UE_DEPRECATED(5.1, "MasterAlpha has been deprecated. Please, use WorldAlpha.")
+	UPROPERTY()
+	float MasterAlpha;
+#endif // WITH_EDITORONLY_DATA
+
 	// Multiplier on the Z-component of velocity and acceleration that is passed to the simulation. Usually from 0.0 to 1.0 to 
 	// reduce the effects of jumping and crouching on the simulation, but it can be higher than 1.0 if you need to exaggerate this motion for some reason.
 	UPROPERTY(EditAnywhere, BlueprintReadWrite, Category = Settings, meta = (ClampMin = "0.0"))
