Uninitialized vars fix

diff --git a/Engine/Plugins/Animation/ControlRig/Source/ControlRig/Private/Units/Hierarchy/RigUnit_Metadata.h b/Engine/Plugins/Animation/ControlRig/Source/ControlRig/Private/Units/Hierarchy/RigUnit_Metadata.h
--- a/Engine/Plugins/Animation/ControlRig/Source/ControlRig/Private/Units/Hierarchy/RigUnit_Metadata.h
+++ b/Engine/Plugins/Animation/ControlRig/Source/ControlRig/Private/Units/Hierarchy/RigUnit_Metadata.h
@@ -432,6 +432,7 @@ struct CONTROLRIG_API FRigUnit_RemoveMetadataTag : public FRigUnitMutable
 	FRigUnit_RemoveMetadataTag()
 		: Item(NAME_None, ERigElementType::Bone)
 		, Tag(NAME_None)
+		, Removed(false)
 		, CachedIndex()
 	{}
 
diff --git a/Engine/Source/Runtime/AnimGraphRuntime/Public/BoneControllers/AnimNode_RigidBody.h b/Engine/Source/Runtime/AnimGraphRuntime/Public/BoneControllers/AnimNode_RigidBody.h
--- a/Engine/Source/Runtime/AnimGraphRuntime/Public/BoneControllers/AnimNode_RigidBody.h
+++ b/Engine/Source/Runtime/AnimGraphRuntime/Public/BoneControllers/AnimNode_RigidBody.h
@@ -68,7 +68,7 @@ struct ANIMGRAPHRUNTIME_API FSimSpaceSettings
 #if WITH_EDITORONLY_DATA
 	UE_DEPRECATED(5.1, "This property has been deprecated. Please, use WorldAlpha.")
 	UPROPERTY()
-	float MasterAlpha;
+	float MasterAlpha = 0.f;
 #endif // WITH_EDITORONLY_DATA
 
 	// Multiplier on the Z-component of velocity and acceleration that is passed to the simulation. Usually from 0.0 to 1.0 to 
diff --git a/Engine/Source/Runtime/Experimental/ChaosSolverEngine/Public/Chaos/ChaosGameplayEventDispatcher.h b/Engine/Source/Runtime/Experimental/ChaosSolverEngine/Public/Chaos/ChaosGameplayEventDispatcher.h
--- a/Engine/Source/Runtime/Experimental/ChaosSolverEngine/Public/Chaos/ChaosGameplayEventDispatcher.h
+++ b/Engine/Source/Runtime/Experimental/ChaosSolverEngine/Public/Chaos/ChaosGameplayEventDispatcher.h
@@ -84,6 +84,8 @@ public:
 		, Orientation(FQuat::Identity)
 		, LinearVelocity(FVector::ZeroVector)
 		, AngularVelocity(FVector::ZeroVector)
+		, Mass(0)
+		, LocalBounds(ForceInitToZero)
 	{}
 
 	/** primitive component involved in the crumble event */
