RBAN - add warning when using a bone-space simulation with an invalid base bone

diff --git a/Engine/Source/Runtime/AnimGraphRuntime/Private/BoneControllers/AnimNode_RigidBody.cpp b/Engine/Source/Runtime/AnimGraphRuntime/Private/BoneControllers/AnimNode_RigidBody.cpp
--- a/Engine/Source/Runtime/AnimGraphRuntime/Private/BoneControllers/AnimNode_RigidBody.cpp
+++ b/Engine/Source/Runtime/AnimGraphRuntime/Private/BoneControllers/AnimNode_RigidBody.cpp
@@ -1775,6 +1775,13 @@ void FAnimNode_RigidBody::InitializeBoneReferences(const FBoneContainer& Require
 		BaseBoneRef.Initialize(RequiredBones);
 	}
 
+	if (!BaseBoneRef.HasValidSetup())
+	{
+		// If the user specified a simulation root that is not used by the skelmesh, issue a warning 
+		// (FAnimNode_RigidBody::IsValidToEvaluate will return false and the simulation will not run)
+		UE_LOG(LogRBAN, Warning, TEXT("FAnimNode_RigidBody: RBAN Simulation Base Bone \'%s\' does not exist on SkeletalMesh %s."), *BaseBoneRef.BoneName.ToString(), *GetNameSafe(RequiredBones.GetSkeletalMeshAsset()));
+	}
+
 	bool bHasInvalidBoneReference = false;
 	for (int32 Index = 0; Index < NumRequiredBoneIndices; ++Index)
 	{
