Fix for crash when removing blend pose from Layered Blend Per Bone
#rnx

diff --git a/Engine/Source/Runtime/AnimGraphRuntime/Public/AnimNodes/AnimNode_LayeredBoneBlend.h b/Engine/Source/Runtime/AnimGraphRuntime/Public/AnimNodes/AnimNode_LayeredBoneBlend.h
--- a/Engine/Source/Runtime/AnimGraphRuntime/Public/AnimNodes/AnimNode_LayeredBoneBlend.h
+++ b/Engine/Source/Runtime/AnimGraphRuntime/Public/AnimNodes/AnimNode_LayeredBoneBlend.h
@@ -129,17 +129,32 @@ public:
 	void AddPose()
 	{
 		BlendWeights.Add(1.f);
-		BlendMasks.Add(nullptr);
 		new (BlendPoses) FPoseLink();
-		new (LayerSetup) FInputBlendPose();
+
+		if (BlendMode == ELayeredBoneBlendMode::BlendMask) 
+		{ 
+			BlendMasks.Add(nullptr); 
+		}
+		else /*if (BlendMode::BranchFilter)*/ 
+		{ 
+			new (LayerSetup) FInputBlendPose(); 
+		}
 	}
 
 	void RemovePose(int32 PoseIndex)
 	{
 		BlendWeights.RemoveAt(PoseIndex);
-		BlendMasks.RemoveAt(PoseIndex);
 		BlendPoses.RemoveAt(PoseIndex);
-		LayerSetup.RemoveAt(PoseIndex);
+
+		if (BlendMasks.IsValidIndex(PoseIndex)) 
+		{ 
+			BlendMasks.RemoveAt(PoseIndex); 
+		}
+
+		if (LayerSetup.IsValidIndex(PoseIndex)) 
+		{ 
+			LayerSetup.RemoveAt(PoseIndex); 
+		}
 	}
 
 	// Invalidate the cached per-bone blend weights from the skeleton
