Fix for crash in FAnimNode_LayeredBlendPerBone::Update when weights array is empty
#rnx

diff --git a/Engine/Source/Runtime/AnimGraphRuntime/Private/AnimNodes/AnimNode_LayeredBoneBlend.cpp b/Engine/Source/Runtime/AnimGraphRuntime/Private/AnimNodes/AnimNode_LayeredBoneBlend.cpp
--- a/Engine/Source/Runtime/AnimGraphRuntime/Private/AnimNodes/AnimNode_LayeredBoneBlend.cpp
+++ b/Engine/Source/Runtime/AnimGraphRuntime/Private/AnimNodes/AnimNode_LayeredBoneBlend.cpp
@@ -166,7 +166,7 @@ void FAnimNode_LayeredBoneBlend::Update_AnyThread(const FAnimationUpdateContext&
 					FAnimationRuntime::UpdateDesiredBoneWeight(DesiredBoneBlendWeights, CurrentBoneBlendWeights, BlendWeights);
 					bHasRelevantPoses = true;
 
-					if(bBlendRootMotionBasedOnRootBone)
+					if(bBlendRootMotionBasedOnRootBone && !CurrentBoneBlendWeights.IsEmpty())
 					{
 						const float NewRootMotionWeight = CurrentBoneBlendWeights[0].BlendWeight;
 						if(NewRootMotionWeight > ZERO_ANIMWEIGHT_THRESH)
