Prevent out of range errorin FAnimNode_RandomPlayer, due to precision errors
GitHub #9137

# robomerge 5.1

diff --git a/Engine/Source/Runtime/AnimGraphRuntime/Private/AnimNodes/AnimNode_RandomPlayer.cpp b/Engine/Source/Runtime/AnimGraphRuntime/Private/AnimNodes/AnimNode_RandomPlayer.cpp
--- a/Engine/Source/Runtime/AnimGraphRuntime/Private/AnimNodes/AnimNode_RandomPlayer.cpp
+++ b/Engine/Source/Runtime/AnimGraphRuntime/Private/AnimNodes/AnimNode_RandomPlayer.cpp
@@ -49,7 +49,6 @@ void FAnimNode_RandomPlayer::Initialize_AnyThread(const FAnimationInitializeCont
 	}
 
 	NormalizedPlayChances.Empty(NormalizedPlayChances.Num());
-	NormalizedPlayChances.AddUninitialized(NumValidEntries + 1);
 
 	// Sanitize the data and sum up the range of the random chances so that
 	// we can normalize the individual chances below.
@@ -90,7 +89,8 @@ void FAnimNode_RandomPlayer::Initialize_AnyThread(const FAnimationInitializeCont
 			CurrentChance += ValidEntries[Idx]->ChanceToPlay / SumChances;
 			NormalizedPlayChances[Idx] = CurrentChance;
 		}
-		NormalizedPlayChances[NumValidEntries] = 1.0f;
+		// Remove rounding errors (possibly slightly padding out the chance of the last item)
+		NormalizedPlayChances[NumValidEntries - 1] = 1.0f;
 	}
 
 	// Initialize random stream and pick first entry
