Fix for MakeDynamicAdditive
Missing name on FScopedExpectsAdditiveOverride

diff --git a/Engine/Source/Runtime/AnimGraphRuntime/Private/AnimNodes/AnimNode_MakeDynamicAdditive.cpp b/Engine/Source/Runtime/AnimGraphRuntime/Private/AnimNodes/AnimNode_MakeDynamicAdditive.cpp
--- a/Engine/Source/Runtime/AnimGraphRuntime/Private/AnimNodes/AnimNode_MakeDynamicAdditive.cpp
+++ b/Engine/Source/Runtime/AnimGraphRuntime/Private/AnimNodes/AnimNode_MakeDynamicAdditive.cpp
@@ -40,7 +40,7 @@ void FAnimNode_MakeDynamicAdditive::Update_AnyThread(const FAnimationUpdateConte
 void FAnimNode_MakeDynamicAdditive::Evaluate_AnyThread(FPoseContext& Output)
 {
 	DECLARE_SCOPE_HIERARCHICAL_COUNTER_ANIMNODE(Evaluate_AnyThread)
-	FScopedExpectsAdditiveOverride(Output, false);
+	FScopedExpectsAdditiveOverride ScopedExpectsAdditiveOverride(Output, false);
 	FPoseContext BaseEvalContext(Output);
 
 	Base.Evaluate(BaseEvalContext);
