Fix for FAnimNode_RandomPlayer::Initialize_AnyThread crash

diff --git a/Engine/Source/Runtime/AnimGraphRuntime/Private/AnimNodes/AnimNode_RandomPlayer.cpp b/Engine/Source/Runtime/AnimGraphRuntime/Private/AnimNodes/AnimNode_RandomPlayer.cpp
--- a/Engine/Source/Runtime/AnimGraphRuntime/Private/AnimNodes/AnimNode_RandomPlayer.cpp
+++ b/Engine/Source/Runtime/AnimGraphRuntime/Private/AnimNodes/AnimNode_RandomPlayer.cpp
@@ -49,6 +49,7 @@ void FAnimNode_RandomPlayer::Initialize_AnyThread(const FAnimationInitializeCont
 	}
 
 	NormalizedPlayChances.Empty(NormalizedPlayChances.Num());
+	NormalizedPlayChances.AddUninitialized(NumValidEntries);
 
 	// Sanitize the data and sum up the range of the random chances so that
 	// we can normalize the individual chances below.
