PR #9253: Added blendspace player library to support editing blendspace players (Contributed by [REDACTED_PERSON])

diff --git a/Engine/Source/Runtime/AnimGraphRuntime/Private/BlendSpacePlayerLibrary.cpp b/Engine/Source/Runtime/AnimGraphRuntime/Private/BlendSpacePlayerLibrary.cpp
new file mode 100644
--- /dev/null
+++ b/Engine/Source/Runtime/AnimGraphRuntime/Private/BlendSpacePlayerLibrary.cpp
@@ -0,0 +1,186 @@
+// Copyright Epic Games, Inc. All Rights Reserved.
+
+#include "BlendSpacePlayerLibrary.h"
+#include "Animation/AnimNode_Inertialization.h"
+#include "AnimNodes/AnimNode_BlendSpacePlayer.h"
+
+DEFINE_LOG_CATEGORY_STATIC(LogBlendSpacePlayerLibrary, Verbose, All);
+
+FBlendSpacePlayerReference UBlendSpacePlayerLibrary::ConvertToBlendSpacePlayer(const FAnimNodeReference& Node,
+	EAnimNodeReferenceConversionResult& Result)
+{
+	return FAnimNodeReference::ConvertToType<FBlendSpacePlayerReference>(Node, Result);
+}
+
+FBlendSpacePlayerReference UBlendSpacePlayerLibrary::SetBlendSpace(
+	const FBlendSpacePlayerReference& BlendSpacePlayer, UBlendSpace* BlendSpace)
+{
+	BlendSpacePlayer.CallAnimNodeFunction<FAnimNode_BlendSpacePlayer>(
+	TEXT("SetBlendSpace"),
+	[BlendSpace](FAnimNode_BlendSpacePlayer& InBlendSpacePlayer)
+	{
+		if (!InBlendSpacePlayer.SetBlendSpace(BlendSpace))
+		{
+			UE_LOG(LogBlendSpacePlayerLibrary, Warning, TEXT("Could not set blendspace on blendspace player, value is not dynamic. Set it as Always Dynamic."));
+		}
+	});
+
+	return BlendSpacePlayer;
+}
+
+FBlendSpacePlayerReference UBlendSpacePlayerLibrary::SetBlendSpaceWithInertialBlending(
+	const FAnimUpdateContext& UpdateContext, const FBlendSpacePlayerReference& BlendSpacePlayer,
+	UBlendSpace* BlendSpace, float BlendTime)
+{
+	BlendSpacePlayer.CallAnimNodeFunction<FAnimNode_BlendSpacePlayer>(
+		TEXT("SetBlendSpaceWithInertialBlending"),
+		[BlendSpace, &UpdateContext, BlendTime](FAnimNode_BlendSpacePlayer& InBlendSpacePlayer)
+		{
+			const UBlendSpace* CurrentBlendSpace = InBlendSpacePlayer.GetBlendSpace();
+			const bool bBlendSpaceChanged = (CurrentBlendSpace != BlendSpace);
+
+			if (!InBlendSpacePlayer.SetBlendSpace(BlendSpace))
+			{
+				UE_LOG(LogBlendSpacePlayerLibrary, Warning, TEXT("Could not set blendspace on blendspace player, value is not dynamic. Set it as Always Dynamic."));
+			}
+
+			if (bBlendSpaceChanged && BlendTime > 0.0f)
+			{
+				if (const FAnimationUpdateContext* AnimationUpdateContext = UpdateContext.GetContext())
+				{
+					if (UE::Anim::IInertializationRequester* InertializationRequester = AnimationUpdateContext->GetMessage<UE::Anim::IInertializationRequester>())
+					{
+						InertializationRequester->RequestInertialization(BlendTime);
+					}
+				}
+				else
+				{
+					UE_LOG(LogBlendSpacePlayerLibrary, Warning, TEXT("SetBlendSpaceWithIntertialBlending called with invalid context."));
+				}
+			}
+		});
+	
+	return BlendSpacePlayer;
+}
+
+FBlendSpacePlayerReference UBlendSpacePlayerLibrary::SetResetPlayTimeWhenBlendSpaceChanges(
+	const FBlendSpacePlayerReference& BlendSpacePlayer, bool bReset)
+{
+	BlendSpacePlayer.CallAnimNodeFunction<FAnimNode_BlendSpacePlayer>(
+		TEXT("SetResetPlayTimeWhenBlendSpaceChanges"),
+		[bReset](FAnimNode_BlendSpacePlayer& InBlendSpacePlayer)
+		{
+			if (!InBlendSpacePlayer.SetResetPlayTimeWhenBlendSpaceChanges(bReset))
+			{
+				UE_LOG(LogBlendSpacePlayerLibrary, Warning, TEXT("Could not set reset play time when blend space changes on blendspace player, value is not dynamic. Set it as Always Dynamic."));
+			}
+		});
+
+	return BlendSpacePlayer;
+}
+
+
+FBlendSpacePlayerReference UBlendSpacePlayerLibrary::SetPlayRate(const FBlendSpacePlayerReference& BlendSpacePlayer,
+	float PlayRate)
+{
+	BlendSpacePlayer.CallAnimNodeFunction<FAnimNode_BlendSpacePlayer>(
+		TEXT("SetPlayRate"),
+		[PlayRate](FAnimNode_BlendSpacePlayer& InBlendSpacePlayer)
+		{
+			if (!InBlendSpacePlayer.SetPlayRate(PlayRate))
+			{
+				UE_LOG(LogBlendSpacePlayerLibrary, Warning, TEXT("Could not set play rate on blendspace player, value is not dynamic. Set it as Always Dynamic."));
+			}
+		});
+
+	return BlendSpacePlayer;
+}
+
+FBlendSpacePlayerReference UBlendSpacePlayerLibrary::SetLoop(const FBlendSpacePlayerReference& BlendSpacePlayer,
+	bool bLoop)
+{
+	BlendSpacePlayer.CallAnimNodeFunction<FAnimNode_BlendSpacePlayer>(
+		TEXT("SetLoop"),
+		[bLoop](FAnimNode_BlendSpacePlayer& InBlendSpacePlayer)
+		{
+			if (!InBlendSpacePlayer.SetLoop(bLoop))
+			{
+				UE_LOG(LogBlendSpacePlayerLibrary, Warning, TEXT("Could not set loop on blendspace player, value is not dynamic. Set it as Always Dynamic."));
+			}
+		});
+
+	return BlendSpacePlayer;
+}
+
+UBlendSpace* UBlendSpacePlayerLibrary::GetBlendSpace(const FBlendSpacePlayerReference& BlendSpacePlayer)
+{
+	UBlendSpace* BlendSpace = nullptr;
+	BlendSpacePlayer.CallAnimNodeFunction<FAnimNode_BlendSpacePlayer>(
+		TEXT("GetBlendSpace"),
+		[&BlendSpace](FAnimNode_BlendSpacePlayer& InBlendSpacePlayer)
+		{
+			BlendSpace = InBlendSpacePlayer.GetBlendSpace();
+		});
+	return BlendSpace;
+}
+
+FVector UBlendSpacePlayerLibrary::GetPosition(const FBlendSpacePlayerReference& BlendSpacePlayer)
+{
+	FVector Position = FVector::ZeroVector;
+	BlendSpacePlayer.CallAnimNodeFunction<FAnimNode_BlendSpacePlayer>(
+		TEXT("GetPosition"),
+		[&Position](FAnimNode_BlendSpacePlayer& InBlendSpacePlayer)
+		{
+			Position = InBlendSpacePlayer.GetPosition();
+		});
+	return Position;
+}
+
+float UBlendSpacePlayerLibrary::GetStartPosition(const FBlendSpacePlayerReference& BlendSpacePlayer)
+{
+	float StartPosition = 0.f;
+	BlendSpacePlayer.CallAnimNodeFunction<FAnimNode_BlendSpacePlayer>(
+		TEXT("GetStartPosition"),
+		[&StartPosition](FAnimNode_BlendSpacePlayer& InBlendSpacePlayer)
+		{
+			StartPosition = InBlendSpacePlayer.GetStartPosition();
+		});
+	return StartPosition;
+}
+
+float UBlendSpacePlayerLibrary::GetPlayRate(const FBlendSpacePlayerReference& BlendSpacePlayer)
+{
+	float PlayRate = 0.f;
+	BlendSpacePlayer.CallAnimNodeFunction<FAnimNode_BlendSpacePlayer>(
+		TEXT("GetPlayRate"),
+		[&PlayRate](FAnimNode_BlendSpacePlayer& InBlendSpacePlayer)
+		{
+			PlayRate = InBlendSpacePlayer.GetPlayRate();
+		});
+	return PlayRate;
+}
+
+bool UBlendSpacePlayerLibrary::GetLoop(const FBlendSpacePlayerReference& BlendSpacePlayer)
+{
+	bool bLoop = false;
+	BlendSpacePlayer.CallAnimNodeFunction<FAnimNode_BlendSpacePlayer>(
+		TEXT("GetLoop"),
+		[&bLoop](FAnimNode_BlendSpacePlayer& InBlendSpacePlayer)
+		{
+			bLoop = InBlendSpacePlayer.GetLoop();
+		});
+	return bLoop;
+}
+
+bool UBlendSpacePlayerLibrary::ShouldResetPlayTimeWhenBlendSpaceChanges(
+	const FBlendSpacePlayerReference& BlendSpacePlayer)
+{
+	bool bResetPlayTimeWhenBlendSpaceChanges = false;
+	BlendSpacePlayer.CallAnimNodeFunction<FAnimNode_BlendSpacePlayer>(
+		TEXT("ShouldResetPlayTimeWhenBlendSpaceChanges"),
+		[&bResetPlayTimeWhenBlendSpaceChanges](FAnimNode_BlendSpacePlayer& InBlendSpacePlayer)
+		{
+			bResetPlayTimeWhenBlendSpaceChanges = InBlendSpacePlayer.ShouldResetPlayTimeWhenBlendSpaceChanges();
+		});
+	return bResetPlayTimeWhenBlendSpaceChanges;
+}
diff --git a/Engine/Source/Runtime/AnimGraphRuntime/Public/BlendSpacePlayerLibrary.h b/Engine/Source/Runtime/AnimGraphRuntime/Public/BlendSpacePlayerLibrary.h
new file mode 100644
--- /dev/null
+++ b/Engine/Source/Runtime/AnimGraphRuntime/Public/BlendSpacePlayerLibrary.h
@@ -0,0 +1,89 @@
+// Copyright Epic Games, Inc. All Rights Reserved.
+
+#pragma once
+
+#include "CoreMinimal.h"
+#include "Animation/AnimExecutionContext.h"
+#include "Kismet/BlueprintFunctionLibrary.h"
+#include "Animation/AnimNodeReference.h"
+#include "BlendSpacePlayerLibrary.generated.h"
+
+struct FAnimNode_BlendSpacePlayer;
+
+USTRUCT(BlueprintType)
+struct FBlendSpacePlayerReference : public FAnimNodeReference
+{
+	GENERATED_BODY()
+
+	typedef FAnimNode_BlendSpacePlayer FInternalNodeType;
+};
+
+/**
+ * Exposes operations to be performed on a blend space player anim node.
+ */
+UCLASS()
+class ANIMGRAPHRUNTIME_API UBlendSpacePlayerLibrary : public UBlueprintFunctionLibrary
+{
+	GENERATED_BODY()
+
+public:
+	/** Get a blend space player context from an anim node context. */
+	UFUNCTION(BlueprintCallable, Category = "Blend Space Player", meta = (BlueprintThreadSafe, ExpandEnumAsExecs = "Result"))
+	static FBlendSpacePlayerReference ConvertToBlendSpacePlayer(const FAnimNodeReference& Node, EAnimNodeReferenceConversionResult& Result);
+
+	/** Get a blend space player context from an anim node context (pure). */
+	UFUNCTION(BlueprintPure, Category = "Blend Space Player", meta = (BlueprintThreadSafe, DisplayName = "Convert to Blend Space Player"))
+	static void ConvertToBlendSpacePlayerPure(const FAnimNodeReference& Node, FBlendSpacePlayerReference& BlendSpacePlayer, bool& Result)
+	{
+		EAnimNodeReferenceConversionResult ConversionResult;
+		BlendSpacePlayer = ConvertToBlendSpacePlayer(Node, ConversionResult);
+		Result = (ConversionResult == EAnimNodeReferenceConversionResult::Succeeded);
+	}
+
+	/** Set the current BlendSpace of the blend space player. */
+	UFUNCTION(BlueprintCallable, Category = "Blend Space Player", meta = (BlueprintThreadSafe))
+	static FBlendSpacePlayerReference SetBlendSpace(const FBlendSpacePlayerReference& BlendSpacePlayer, UBlendSpace* BlendSpace);
+
+	/** Set the current BlendSpace of the blend space player with an interial blend time. */
+	UFUNCTION(BlueprintCallable, Category = "Blend Space Player", meta = (BlueprintThreadSafe))
+	static FBlendSpacePlayerReference SetBlendSpaceWithInertialBlending(const FAnimUpdateContext& UpdateContext, const FBlendSpacePlayerReference& BlendSpacePlayer, UBlendSpace* BlendSpace, float BlendTime = 0.2f);
+
+	/** Set whether the current play time should reset when BlendSpace changes of the blend space player. */
+	UFUNCTION(BlueprintCallable, Category = "Blend Space Player", meta = (BlueprintThreadSafe))
+	static FBlendSpacePlayerReference SetResetPlayTimeWhenBlendSpaceChanges(const FBlendSpacePlayerReference& BlendSpacePlayer, bool bReset);
+	
+	/** Set the play rate of the blend space player. */
+	UFUNCTION(BlueprintCallable, Category = "Blend Space Player", meta = (BlueprintThreadSafe))
+	static FBlendSpacePlayerReference SetPlayRate(const FBlendSpacePlayerReference& BlendSpacePlayer, float PlayRate);
+
+	/** Set the loop of the blend space player. */
+	UFUNCTION(BlueprintCallable, Category = "Blend Space Player", meta = (BlueprintThreadSafe))
+	static FBlendSpacePlayerReference SetLoop(const FBlendSpacePlayerReference& BlendSpacePlayer, bool bLoop);
+
+	/** Get the current BlendSpace of the blend space player. */
+	UFUNCTION(BlueprintPure, Category = "Blend Space Player", meta = (BlueprintThreadSafe))
+	static UBlendSpace* GetBlendSpace(const FBlendSpacePlayerReference& BlendSpacePlayer);
+
+	/** Get the current position of the blend space player. */
+	UFUNCTION(BlueprintPure, Category = "Blend Space Player", meta = (BlueprintThreadSafe))
+	static FVector GetPosition(const FBlendSpacePlayerReference& BlendSpacePlayer);
+
+	/** Get the current start position of the blend space player. */
+	UFUNCTION(BlueprintPure, Category = "Blend Space Player", meta = (BlueprintThreadSafe))
+	static float GetStartPosition(const FBlendSpacePlayerReference& BlendSpacePlayer);
+
+	/** Get the current play rate of the blend space player. */
+	UFUNCTION(BlueprintPure, Category = "Blend Space Player", meta = (BlueprintThreadSafe))
+	static float GetPlayRate(const FBlendSpacePlayerReference& BlendSpacePlayer);
+
+	/** Get the current loop of the blend space player.  */
+	UFUNCTION(BlueprintPure, Category = "Blend Space Player", meta = (BlueprintThreadSafe))
+	static bool GetLoop(const FBlendSpacePlayerReference& BlendSpacePlayer);
+
+	/** Get the current value of whether the current play time should reset when BlendSpace changes of the blend space player. */
+	UFUNCTION(BlueprintPure, Category = "Blend Space Player", meta = (BlueprintThreadSafe))
+	static bool ShouldResetPlayTimeWhenBlendSpaceChanges(const FBlendSpacePlayerReference& BlendSpacePlayer); 
+};
+
+
+
