After five years, it's probably time to remove the 'experimental' flag from the Constraint anim node.

#rnx

diff --git a/Engine/Source/Runtime/AnimGraphRuntime/Public/BoneControllers/AnimNode_Constraint.h b/Engine/Source/Runtime/AnimGraphRuntime/Public/BoneControllers/AnimNode_Constraint.h
--- a/Engine/Source/Runtime/AnimGraphRuntime/Public/BoneControllers/AnimNode_Constraint.h
+++ b/Engine/Source/Runtime/AnimGraphRuntime/Public/BoneControllers/AnimNode_Constraint.h
@@ -70,7 +70,7 @@ struct FConstraint
 /**
  *	Constraint node to parent or world transform for rotation/translation
  */
-USTRUCT(Experimental)
+USTRUCT()
 struct ANIMGRAPHRUNTIME_API FAnimNode_Constraint : public FAnimNode_SkeletalControlBase
 {
 	GENERATED_USTRUCT_BODY()
