Fix assert on CVars when accessed from the async loading thread in FAnimNode_AnimDynamics::HasPreUpdate

diff --git a/Engine/Source/Runtime/AnimGraphRuntime/Private/BoneControllers/AnimNode_AnimDynamics.cpp b/Engine/Source/Runtime/AnimGraphRuntime/Private/BoneControllers/AnimNode_AnimDynamics.cpp
--- a/Engine/Source/Runtime/AnimGraphRuntime/Private/BoneControllers/AnimNode_AnimDynamics.cpp
+++ b/Engine/Source/Runtime/AnimGraphRuntime/Private/BoneControllers/AnimNode_AnimDynamics.cpp
@@ -1006,11 +1006,11 @@ void FAnimNode_AnimDynamics::UpdateLimits(FComponentSpacePoseContext& Output)
 
 bool FAnimNode_AnimDynamics::HasPreUpdate() const
 {
-	if(CVarEnableDynamics.GetValueOnGameThread() == 1)
+	if(CVarEnableDynamics.GetValueOnAnyThread() == 1)
 	{
-		return (CVarEnableWind.GetValueOnGameThread() == 1 && (bEnableWind || bWindWasEnabled))
+		return (CVarEnableWind.GetValueOnAnyThread() == 1 && (bEnableWind || bWindWasEnabled))
 #if ENABLE_ANIM_DRAW_DEBUG
-				|| (CVarShowDebug.GetValueOnGameThread() == 1 && !CVarDebugBone.GetValueOnGameThread().IsEmpty())
+				|| (CVarShowDebug.GetValueOnAnyThread() == 1 && !CVarDebugBone.GetValueOnAnyThread().IsEmpty())
 #endif
 				;
 	}
