GitHub 8512 : Added UAnimSequenceBase::ComputePlayRateFromDuration

diff --git a/Engine/Source/Runtime/AnimGraphRuntime/Private/SequencePlayerLibrary.cpp b/Engine/Source/Runtime/AnimGraphRuntime/Private/SequencePlayerLibrary.cpp
--- a/Engine/Source/Runtime/AnimGraphRuntime/Private/SequencePlayerLibrary.cpp
+++ b/Engine/Source/Runtime/AnimGraphRuntime/Private/SequencePlayerLibrary.cpp
@@ -178,4 +178,22 @@ bool USequencePlayerLibrary::GetLoopAnimation(const FSequencePlayerReference& Se
 		});
 
 	return bLoopAnimation;
-}
\ No newline at end of file
+}
+
+float USequencePlayerLibrary::ComputePlayRateFromDuration(const FSequencePlayerReference& SequencePlayer, float Duration /* = 1.0f */)
+{
+	float PlayRate = 1.f;
+	if (Duration > 0.f)
+	{
+	SequencePlayer.CallAnimNodeFunction<FAnimNode_SequencePlayer>(
+		TEXT("GetPlayRate"),
+		[&PlayRate, Duration](const FAnimNode_SequencePlayer& InSequencePlayer)
+		{
+			if (const UAnimSequenceBase* Sequence = InSequencePlayer.GetSequence())
+			{
+				PlayRate = Sequence->GetPlayLength() / Duration;
+			}
+		});
+	}
+	return PlayRate;
+}
diff --git a/Engine/Source/Runtime/AnimGraphRuntime/Public/SequencePlayerLibrary.h b/Engine/Source/Runtime/AnimGraphRuntime/Public/SequencePlayerLibrary.h
--- a/Engine/Source/Runtime/AnimGraphRuntime/Public/SequencePlayerLibrary.h
+++ b/Engine/Source/Runtime/AnimGraphRuntime/Public/SequencePlayerLibrary.h
@@ -87,4 +87,8 @@ public:
 	/** Get the looping state of the sequence player */
 	UFUNCTION(BlueprintPure, Category = "Animation|Sequences", meta = (BlueprintThreadSafe))
 	static bool GetLoopAnimation(const FSequencePlayerReference& SequencePlayer);
+
+	/** Returns the Play Rate to provide when playing this animation if a specific animation duration is desired */
+	UFUNCTION(BlueprintPure, Category = "Animation|Sequences", meta = (BlueprintThreadSafe))
+	static float ComputePlayRateFromDuration(const FSequencePlayerReference& SequencePlayer, float Duration = 1.f);
 };
\ No newline at end of file
