Added the ability to specify blend profiles/masks to the layered blend per bone node

Enables use of layered blends in template animation blueprints

diff --git a/Engine/Source/Editor/AnimGraph/Private/AnimGraphNode_LayeredBoneBlend.cpp b/Engine/Source/Editor/AnimGraph/Private/AnimGraphNode_LayeredBoneBlend.cpp
--- a/Engine/Source/Editor/AnimGraph/Private/AnimGraphNode_LayeredBoneBlend.cpp
+++ b/Engine/Source/Editor/AnimGraph/Private/AnimGraphNode_LayeredBoneBlend.cpp
@@ -169,13 +169,13 @@ void UAnimGraphNode_LayeredBoneBlend::ValidateAnimNodeDuringCompilation(class US
 		for (int32 MaskIndex = 0; MaskIndex < NumBlendMasks; ++MaskIndex)
 		{
 			const UBlendProfile* BlendMask = Node.BlendMasks[MaskIndex];
-			if (BlendMask == nullptr)
+			if (BlendMask == nullptr && !GetAnimBlueprint()->bIsTemplate)
 			{
 				MessageLog.Error(*FText::Format(LOCTEXT("LayeredBlendNullMask", "@@ has null BlendMask for Blend Pose {0}. "), FText::AsNumber(MaskIndex)).ToString(), this, BlendMask);
 				bCompilationError = true;
-				continue;
 			}
-			else if (BlendMask->Mode != EBlendProfileMode::BlendMask)
+			
+			if (BlendMask && BlendMask->Mode != EBlendProfileMode::BlendMask)
 			{
 				MessageLog.Error(*FText::Format(LOCTEXT("LayeredBlendProfileModeError", "@@ is using a BlendProfile(@@) without a BlendMask mode for Blend Pose {0}. "), FText::AsNumber(MaskIndex)).ToString(), this, BlendMask);
 				bCompilationError = true;
diff --git a/Engine/Source/Runtime/AnimGraphRuntime/Private/AnimNodes/AnimNode_LayeredBoneBlend.cpp b/Engine/Source/Runtime/AnimGraphRuntime/Private/AnimNodes/AnimNode_LayeredBoneBlend.cpp
--- a/Engine/Source/Runtime/AnimGraphRuntime/Private/AnimNodes/AnimNode_LayeredBoneBlend.cpp
+++ b/Engine/Source/Runtime/AnimGraphRuntime/Private/AnimNodes/AnimNode_LayeredBoneBlend.cpp
@@ -300,3 +300,14 @@ void FAnimNode_LayeredBoneBlend::GatherDebugData(FNodeDebugData& DebugData)
 		BlendPoses[ChildIndex].GatherDebugData(DebugData.BranchFlow(BlendWeights[ChildIndex]));
 	}
 }
+
+void FAnimNode_LayeredBoneBlend::SetBlendMask(int32 InPoseIndex, UBlendProfile* InBlendMask)
+{
+	check(BlendMode == ELayeredBoneBlendMode::BlendMask);
+	check(BlendPoses.IsValidIndex(InPoseIndex));
+	check(BlendMasks.IsValidIndex(InPoseIndex));
+
+	BlendMasks[InPoseIndex] = InBlendMask;
+
+	InvalidatePerBoneBlendWeights();
+}
\ No newline at end of file
diff --git a/Engine/Source/Runtime/AnimGraphRuntime/Private/LayeredBoneBlendLibrary.cpp b/Engine/Source/Runtime/AnimGraphRuntime/Private/LayeredBoneBlendLibrary.cpp
new file mode 100644
--- /dev/null
+++ b/Engine/Source/Runtime/AnimGraphRuntime/Private/LayeredBoneBlendLibrary.cpp
@@ -0,0 +1,75 @@
+// Copyright Epic Games, Inc. All Rights Reserved.
+
+#include "LayeredBoneBlendLibrary.h"
+#include "Animation/AnimInstanceProxy.h"
+#include "AnimNodes/AnimNode_LayeredBoneBlend.h"
+
+FLayeredBoneBlendReference ULayeredBoneBlendLibrary::ConvertToLayeredBoneBlend(const FAnimNodeReference& Node, EAnimNodeReferenceConversionResult& Result)
+{
+	return FAnimNodeReference::ConvertToType<FLayeredBoneBlendReference>(Node, Result);
+}
+
+int32 ULayeredBoneBlendLibrary::GetNumPoses(const FLayeredBoneBlendReference& LayeredBoneBlend)
+{
+	int32 NumPoses = 0;
+	
+	LayeredBoneBlend.CallAnimNodeFunction<FAnimNode_LayeredBoneBlend>(
+		TEXT("GetNumPoses"),
+		[&NumPoses](const FAnimNode_LayeredBoneBlend& InLayeredBoneBlend)
+		{
+			NumPoses = InLayeredBoneBlend.BlendPoses.Num();
+		});
+
+	return NumPoses;
+}
+
+FLayeredBoneBlendReference ULayeredBoneBlendLibrary::SetBlendMask(const FAnimUpdateContext& UpdateContext, const FLayeredBoneBlendReference& LayeredBoneBlend, int32 PoseIndex, FName BlendMaskName)
+{
+	LayeredBoneBlend.CallAnimNodeFunction<FAnimNode_LayeredBoneBlend>(
+		TEXT("SetBlendMask"),
+		[&UpdateContext, PoseIndex, BlendMaskName](FAnimNode_LayeredBoneBlend& InLayeredBoneBlend)
+		{
+			const FAnimationUpdateContext* AnimationUpdateContext = UpdateContext.GetContext();
+			if(AnimationUpdateContext == nullptr)
+			{
+				UE_LOG(LogAnimation, Warning, TEXT("SetBlendMask: Invalid update context provided."));
+				return;
+			}
+
+			USkeleton* Skeleton = AnimationUpdateContext->AnimInstanceProxy->GetSkeleton();
+			if(Skeleton == nullptr)
+			{
+				UE_LOG(LogAnimation, Warning, TEXT("SetBlendMask: Invalid skeleton."));
+				return;
+			}
+
+			if(InLayeredBoneBlend.BlendMode != ELayeredBoneBlendMode::BlendMask)
+			{
+				UE_LOG(LogAnimation, Warning, TEXT("SetBlendMask: Node is not set to use blend masks."), *BlendMaskName.ToString());
+				return;
+			}
+
+			if(!InLayeredBoneBlend.BlendPoses.IsValidIndex(PoseIndex))
+			{
+				UE_LOG(LogAnimation, Warning, TEXT("SetBlendMask: Invalid pose index %d provided (%d present)."), PoseIndex, InLayeredBoneBlend.BlendPoses.Num());
+				return;
+			}
+			
+			UBlendProfile* BlendMask = Skeleton->GetBlendProfile(BlendMaskName);
+			if(BlendMask == nullptr)
+			{
+				UE_LOG(LogAnimation, Warning, TEXT("SetBlendMask: Blend mask '%s' not found."), *BlendMaskName.ToString());
+				return;
+			}
+
+			if(BlendMask->Mode != EBlendProfileMode::BlendMask)
+			{
+				UE_LOG(LogAnimation, Warning, TEXT("SetBlendMask: Blend mask '%s' does not use BlendMask mode."), *BlendMaskName.ToString());
+				return;
+			}
+
+			InLayeredBoneBlend.SetBlendMask(PoseIndex, BlendMask);
+		});
+
+	return LayeredBoneBlend;
+}
\ No newline at end of file
diff --git a/Engine/Source/Runtime/AnimGraphRuntime/Public/AnimNodes/AnimNode_LayeredBoneBlend.h b/Engine/Source/Runtime/AnimGraphRuntime/Public/AnimNodes/AnimNode_LayeredBoneBlend.h
--- a/Engine/Source/Runtime/AnimGraphRuntime/Public/AnimNodes/AnimNode_LayeredBoneBlend.h
+++ b/Engine/Source/Runtime/AnimGraphRuntime/Public/AnimNodes/AnimNode_LayeredBoneBlend.h
@@ -157,6 +157,9 @@ public:
 		}
 	}
 
+	// Set the blend mask for the specified input pose
+	void SetBlendMask(int32 InPoseIndex, UBlendProfile* InBlendMask);
+	
 	// Invalidate the cached per-bone blend weights from the skeleton
 	void InvalidatePerBoneBlendWeights() { RequiredBonesSerialNumber = 0; SkeletonGuid = FGuid(); VirtualBoneGuid = FGuid(); }
 	
diff --git a/Engine/Source/Runtime/AnimGraphRuntime/Public/LayeredBoneBlendLibrary.h b/Engine/Source/Runtime/AnimGraphRuntime/Public/LayeredBoneBlendLibrary.h
new file mode 100644
--- /dev/null
+++ b/Engine/Source/Runtime/AnimGraphRuntime/Public/LayeredBoneBlendLibrary.h
@@ -0,0 +1,56 @@
+﻿// Copyright Epic Games, Inc. All Rights Reserved.
+
+#pragma once
+
+#include "CoreMinimal.h"
+#include "Animation/AnimExecutionContext.h"
+#include "Kismet/BlueprintFunctionLibrary.h"
+#include "Animation/AnimNodeReference.h"
+#include "LayeredBoneBlendLibrary.generated.h"
+
+struct FAnimNode_LayeredBoneBlend;
+
+USTRUCT(BlueprintType)
+struct FLayeredBoneBlendReference : public FAnimNodeReference
+{
+	GENERATED_BODY()
+
+	typedef FAnimNode_LayeredBoneBlend FInternalNodeType;
+};
+
+/**
+ * Exposes operations to be performed on a layered bone blend anim node.
+ */
+UCLASS()
+class ANIMGRAPHRUNTIME_API ULayeredBoneBlendLibrary : public UBlueprintFunctionLibrary
+{
+	GENERATED_BODY()
+
+public:
+	/** Get a layered bone blend context from an anim node context. */
+	UFUNCTION(BlueprintCallable, Category = "Layered Bone Blend", meta = (BlueprintThreadSafe, ExpandEnumAsExecs = "Result"))
+	static FLayeredBoneBlendReference ConvertToLayeredBoneBlend(const FAnimNodeReference& Node, EAnimNodeReferenceConversionResult& Result);
+
+	/** Get a layered bone blend context from an anim node context (pure). */
+	UFUNCTION(BlueprintPure, Category = "Layered Bone Blend", meta = (BlueprintThreadSafe, DisplayName = "Convert to Layered Bone Blend"))
+	static void ConvertToLayeredBlendPerBonePure(const FAnimNodeReference& Node, FLayeredBoneBlendReference& LayeredBoneBlend, bool& Result)
+	{
+		EAnimNodeReferenceConversionResult ConversionResult;
+		LayeredBoneBlend = ConvertToLayeredBoneBlend(Node, ConversionResult);
+		Result = (ConversionResult == EAnimNodeReferenceConversionResult::Succeeded);
+	}
+
+	/** Get the number of poses that a layered bone blend node has (this does not include the base pose) */
+	UFUNCTION(BlueprintPure, Category = "Layered Bone Blend", meta = (BlueprintThreadSafe))
+	static int32 GetNumPoses(const FLayeredBoneBlendReference& LayeredBoneBlend);
+	
+	/**
+	 * Sets the currently-used blend mask for a blended input pose by name.
+	 * @param	UpdateContext		The update context to use. This is used to extract the current skeleton to derive the blend mask from
+	 * @param	LayeredBoneBlend	A reference to the node
+	 * @param	PoseIndex			The pose index to set the blend mask for
+	 * @param	BlendMaskName	The name of the blend mask to use 
+	 */
+	UFUNCTION(BlueprintCallable, Category = "Layered Bone Blend", meta = (BlueprintThreadSafe))
+	static FLayeredBoneBlendReference SetBlendMask(const FAnimUpdateContext& UpdateContext, const FLayeredBoneBlendReference& LayeredBoneBlend, int32 PoseIndex, FName BlendMaskName);
+};
\ No newline at end of file
diff --git a/Engine/Source/Runtime/Engine/Private/Animation/AnimationRuntime.cpp b/Engine/Source/Runtime/Engine/Private/Animation/AnimationRuntime.cpp
--- a/Engine/Source/Runtime/Engine/Private/Animation/AnimationRuntime.cpp
+++ b/Engine/Source/Runtime/Engine/Private/Animation/AnimationRuntime.cpp
@@ -2152,7 +2152,6 @@ void FAnimationRuntime::CreateMaskWeights(TArray<FPerBoneBlendWeight>& BoneBlend
 
 			if (!BlendMask || BlendMask->Mode != EBlendProfileMode::BlendMask)
 			{
-				ensureMsgf(false, TEXT("FAnimationRuntime::CreateMaskWeights BlendMask null or BlendProfile mode is not blend mask.  BlendProfile=%s"), *GetNameSafe(BlendMask));
 				continue;
 			}
 
