[UDN Addition] Add the ability to force blend spaces to snap to input
Contributed by [REDACTED_PERSON]

diff --git a/Engine/Source/Runtime/AnimGraphRuntime/Private/AnimNodes/AnimNode_BlendSpaceGraphBase.cpp b/Engine/Source/Runtime/AnimGraphRuntime/Private/AnimNodes/AnimNode_BlendSpaceGraphBase.cpp
--- a/Engine/Source/Runtime/AnimGraphRuntime/Private/AnimNodes/AnimNode_BlendSpaceGraphBase.cpp
+++ b/Engine/Source/Runtime/AnimGraphRuntime/Private/AnimNodes/AnimNode_BlendSpaceGraphBase.cpp
@@ -115,3 +115,12 @@ void FAnimNode_BlendSpaceGraphBase::SetPreviewPosition(FVector InVector)
 	PreviewPosition = InVector;
 }
 #endif
+
+void FAnimNode_BlendSpaceGraphBase::SnapToPosition(const FVector& NewPosition)
+{
+	const int32 NumAxis = FMath::Min(BlendFilter.FilterPerAxis.Num(), 3);
+	for (int32 Idx = 0; Idx < NumAxis; Idx++)
+	{
+		BlendFilter.FilterPerAxis[Idx].SetToValue(NewPosition[Idx]);
+	}
+}
diff --git a/Engine/Source/Runtime/AnimGraphRuntime/Private/AnimNodes/AnimNode_BlendSpacePlayer.cpp b/Engine/Source/Runtime/AnimGraphRuntime/Private/AnimNodes/AnimNode_BlendSpacePlayer.cpp
--- a/Engine/Source/Runtime/AnimGraphRuntime/Private/AnimNodes/AnimNode_BlendSpacePlayer.cpp
+++ b/Engine/Source/Runtime/AnimGraphRuntime/Private/AnimNodes/AnimNode_BlendSpacePlayer.cpp
@@ -157,6 +157,15 @@ float FAnimNode_BlendSpacePlayerBase::GetTimeFromEnd(float CurrentTime) const
 	return GetBlendSpace() != nullptr ? PlayLength - CurrentTime : 0.0f;
 }
 
+void FAnimNode_BlendSpacePlayerBase::SnapToPosition(const FVector& NewPosition)
+{
+	const int32 NumAxis = FMath::Min(BlendFilter.FilterPerAxis.Num(), 3);
+	for (int32 Idx = 0; Idx < NumAxis; Idx++)
+	{
+		BlendFilter.FilterPerAxis[Idx].SetToValue(NewPosition[Idx]);
+	}
+}
+
 UAnimationAsset* FAnimNode_BlendSpacePlayerBase::GetAnimAsset() const
 {
 	return GetBlendSpace();
diff --git a/Engine/Source/Runtime/AnimGraphRuntime/Private/BlendSpaceLibrary.cpp b/Engine/Source/Runtime/AnimGraphRuntime/Private/BlendSpaceLibrary.cpp
new file mode 100644
--- /dev/null
+++ b/Engine/Source/Runtime/AnimGraphRuntime/Private/BlendSpaceLibrary.cpp
@@ -0,0 +1,53 @@
+// Copyright Epic Games, Inc. All Rights Reserved.
+
+#include "BlendSpaceLibrary.h"
+#include "Animation/AnimNode_Inertialization.h"
+#include "AnimNodes/AnimNode_BlendSpaceGraph.h"
+
+DEFINE_LOG_CATEGORY_STATIC(LogBlendSpaceLibrary, Verbose, All);
+
+FBlendSpaceReference UBlendSpaceLibrary::ConvertToBlendSpace(const FAnimNodeReference& Node,
+	EAnimNodeReferenceConversionResult& Result)
+{
+	return FAnimNodeReference::ConvertToType<FBlendSpaceReference>(Node, Result);
+}
+
+FVector UBlendSpaceLibrary::GetPosition(const FBlendSpaceReference& BlendSpace)
+{
+	FVector Position = FVector::ZeroVector;
+
+	BlendSpace.CallAnimNodeFunction<FAnimNode_BlendSpaceGraph>(
+		TEXT("GetPosition"),
+		[&Position](FAnimNode_BlendSpaceGraph& InBlendSpace)
+		{
+			Position = InBlendSpace.GetPosition();
+		});
+
+	return Position;
+}
+
+FVector UBlendSpaceLibrary::GetFilteredPosition(const FBlendSpaceReference& BlendSpace)
+{
+	FVector FilteredPosition = FVector::ZeroVector;
+
+	BlendSpace.CallAnimNodeFunction<FAnimNode_BlendSpaceGraph>(
+		TEXT("GetFilteredPosition"),
+		[&FilteredPosition](FAnimNode_BlendSpaceGraph& InBlendSpace)
+		{
+			FilteredPosition = InBlendSpace.GetFilteredPosition();
+		});
+
+	return FilteredPosition;
+}
+
+void UBlendSpaceLibrary::SnapToPosition(
+	const FBlendSpaceReference& BlendSpace,
+	FVector NewPosition)
+{
+	BlendSpace.CallAnimNodeFunction<FAnimNode_BlendSpaceGraph>(
+		TEXT("SnapToPosition"),
+		[&NewPosition](FAnimNode_BlendSpaceGraph& InBlendSpace)
+		{
+			InBlendSpace.SnapToPosition(NewPosition);
+		});
+}
diff --git a/Engine/Source/Runtime/AnimGraphRuntime/Private/BlendSpacePlayerLibrary.cpp b/Engine/Source/Runtime/AnimGraphRuntime/Private/BlendSpacePlayerLibrary.cpp
--- a/Engine/Source/Runtime/AnimGraphRuntime/Private/BlendSpacePlayerLibrary.cpp
+++ b/Engine/Source/Runtime/AnimGraphRuntime/Private/BlendSpacePlayerLibrary.cpp
@@ -184,3 +184,15 @@ bool UBlendSpacePlayerLibrary::ShouldResetPlayTimeWhenBlendSpaceChanges(
 		});
 	return bResetPlayTimeWhenBlendSpaceChanges;
 }
+
+void UBlendSpacePlayerLibrary::SnapToPosition(
+	const FBlendSpacePlayerReference& BlendSpacePlayer,
+	FVector NewPosition)
+{
+	BlendSpacePlayer.CallAnimNodeFunction<FAnimNode_BlendSpacePlayer>(
+		TEXT("SnapToPosition"),
+		[&NewPosition](FAnimNode_BlendSpacePlayer& InBlendSpacePlayer)
+		{
+			InBlendSpacePlayer.SnapToPosition(NewPosition);
+		});
+}
diff --git a/Engine/Source/Runtime/AnimGraphRuntime/Public/AnimNodes/AnimNode_BlendSpaceGraphBase.h b/Engine/Source/Runtime/AnimGraphRuntime/Public/AnimNodes/AnimNode_BlendSpaceGraphBase.h
--- a/Engine/Source/Runtime/AnimGraphRuntime/Public/AnimNodes/AnimNode_BlendSpaceGraphBase.h
+++ b/Engine/Source/Runtime/AnimGraphRuntime/Public/AnimNodes/AnimNode_BlendSpaceGraphBase.h
@@ -32,6 +32,9 @@ struct ANIMGRAPHRUNTIME_API FAnimNode_BlendSpaceGraphBase : public FAnimNode_Bas
 	void SetPreviewPosition(FVector InVector);
 #endif
 
+	// Forces the Position to the specified value
+	void SnapToPosition(const FVector& NewPosition);
+
 protected:
 	// The X coordinate to sample in the blendspace
 	UPROPERTY(EditAnywhere, BlueprintReadWrite, Category = Coordinates, meta = (PinShownByDefault, AllowPrivateAccess))
diff --git a/Engine/Source/Runtime/AnimGraphRuntime/Public/AnimNodes/AnimNode_BlendSpacePlayer.h b/Engine/Source/Runtime/AnimGraphRuntime/Public/AnimNodes/AnimNode_BlendSpacePlayer.h
--- a/Engine/Source/Runtime/AnimGraphRuntime/Public/AnimNodes/AnimNode_BlendSpacePlayer.h
+++ b/Engine/Source/Runtime/AnimGraphRuntime/Public/AnimNodes/AnimNode_BlendSpacePlayer.h
@@ -52,6 +52,9 @@ public:
 	// @return the current sample coordinates after going through the filtering
 	FVector GetFilteredPosition() const { return BlendFilter.GetFilterLastOutput(); }
 
+	// Forces the Position to the specified value
+	void SnapToPosition(const FVector& NewPosition);
+
 public:
 
 	// Get the blendspace asset to play
diff --git a/Engine/Source/Runtime/AnimGraphRuntime/Public/BlendSpaceLibrary.h b/Engine/Source/Runtime/AnimGraphRuntime/Public/BlendSpaceLibrary.h
new file mode 100644
--- /dev/null
+++ b/Engine/Source/Runtime/AnimGraphRuntime/Public/BlendSpaceLibrary.h
@@ -0,0 +1,57 @@
+// Copyright Epic Games, Inc. All Rights Reserved.
+
+#pragma once
+
+#include "CoreMinimal.h"
+#include "Animation/AnimExecutionContext.h"
+#include "Kismet/BlueprintFunctionLibrary.h"
+#include "Animation/AnimNodeReference.h"
+#include "BlendSpaceLibrary.generated.h"
+
+struct FAnimNode_BlendSpaceGraph;
+
+USTRUCT(BlueprintType)
+struct FBlendSpaceReference : public FAnimNodeReference
+{
+	GENERATED_BODY()
+
+	typedef FAnimNode_BlendSpaceGraph FInternalNodeType;
+};
+
+/**
+ * Exposes operations to be performed on a blend space anim node.
+ */
+UCLASS()
+class ANIMGRAPHRUNTIME_API UBlendSpaceLibrary : public UBlueprintFunctionLibrary
+{
+	GENERATED_BODY()
+
+public:
+	/** Get a blend space context from an anim node context. */
+	UFUNCTION(BlueprintCallable, Category = "Blend Space", meta = (BlueprintThreadSafe, ExpandEnumAsExecs = "Result"))
+	static FBlendSpaceReference ConvertToBlendSpace(const FAnimNodeReference& Node, EAnimNodeReferenceConversionResult& Result);
+
+	/** Get a blend space context from an anim node context (pure). */
+	UFUNCTION(BlueprintPure, Category = "Blend Space", meta = (BlueprintThreadSafe, DisplayName = "Convert to Blend Space (Pure)"))
+	static void ConvertToBlendSpacePure(const FAnimNodeReference& Node, FBlendSpaceReference& BlendSpace, bool& Result)
+	{
+		EAnimNodeReferenceConversionResult ConversionResult;
+		BlendSpace = ConvertToBlendSpace(Node, ConversionResult);
+		Result = (ConversionResult == EAnimNodeReferenceConversionResult::Succeeded);
+	}
+
+	/** Get the current position of the blend space. */
+	UFUNCTION(BlueprintPure, Category = "Blend Space", meta = (BlueprintThreadSafe))
+	static FVector GetPosition(const FBlendSpaceReference& BlendSpace);
+
+	/** Get the current sample coordinates after going through the filtering. */
+	UFUNCTION(BlueprintPure, Category = "Blend Space", meta = (BlueprintThreadSafe))
+	static FVector GetFilteredPosition(const FBlendSpaceReference& BlendSpace);
+
+	/** Forces the Position to the specified value */
+	UFUNCTION(BlueprintCallable, Category = "Blend Space", meta = (BlueprintThreadSafe))
+	static void SnapToPosition(const FBlendSpaceReference& BlendSpace, FVector NewPosition);
+};
+
+
+
diff --git a/Engine/Source/Runtime/AnimGraphRuntime/Public/BlendSpacePlayerLibrary.h b/Engine/Source/Runtime/AnimGraphRuntime/Public/BlendSpacePlayerLibrary.h
--- a/Engine/Source/Runtime/AnimGraphRuntime/Public/BlendSpacePlayerLibrary.h
+++ b/Engine/Source/Runtime/AnimGraphRuntime/Public/BlendSpacePlayerLibrary.h
@@ -83,6 +83,10 @@ public:
 	/** Get the current value of whether the current play time should reset when BlendSpace changes of the blend space player. */
 	UFUNCTION(BlueprintPure, Category = "Blend Space Player", meta = (BlueprintThreadSafe))
 	static bool ShouldResetPlayTimeWhenBlendSpaceChanges(const FBlendSpacePlayerReference& BlendSpacePlayer); 
+
+	/** Forces the Position to the specified value */
+	UFUNCTION(BlueprintCallable, Category = "Blend Space Player", meta = (BlueprintThreadSafe))
+	static void SnapToPosition(const FBlendSpacePlayerReference& BlendSpacePlayer, FVector NewPosition);
 };
 
 
diff --git a/Engine/Source/Runtime/Engine/Private/Animation/AnimInterpFilter.cpp b/Engine/Source/Runtime/Engine/Private/Animation/AnimInterpFilter.cpp
--- a/Engine/Source/Runtime/Engine/Private/Animation/AnimInterpFilter.cpp
+++ b/Engine/Source/Runtime/Engine/Private/Animation/AnimInterpFilter.cpp
@@ -178,6 +178,15 @@ void FFIRFilterTimeBased::WrapToValue(float Input, float Range)
 	}
 }
 
+void FFIRFilterTimeBased::SetToValue(float Value)
+{
+	LastOutput = Value;
+
+	if (FilterData.Num() > 0)
+	{
+		FilterData[0].Input = Value;
+	}
+}
 
 float FFIRFilterTimeBased::UpdateAndGetFilteredData(float Input, float DeltaTime)
 {
diff --git a/Engine/Source/Runtime/Engine/Public/AnimInterpFilter.h b/Engine/Source/Runtime/Engine/Public/AnimInterpFilter.h
--- a/Engine/Source/Runtime/Engine/Public/AnimInterpFilter.h
+++ b/Engine/Source/Runtime/Engine/Public/AnimInterpFilter.h
@@ -148,6 +148,9 @@ public:
 	// Wraps the internal state by steps of Range so that it is as close as possible to Input
 	void WrapToValue(float Input, float Range);
 
+	// Sets internal state to the specified value
+	void SetToValue(float Input);
+
 	// Filter is considered valid if the WindowDuration is > 0
 	bool IsValid() const { return WindowDuration > 0.f; }
 
