Chaos - move RBAN world object collection to game thread

diff --git a/Engine/Source/Runtime/AnimGraphRuntime/Private/BoneControllers/AnimNode_RigidBody.cpp b/Engine/Source/Runtime/AnimGraphRuntime/Private/BoneControllers/AnimNode_RigidBody.cpp
--- a/Engine/Source/Runtime/AnimGraphRuntime/Private/BoneControllers/AnimNode_RigidBody.cpp
+++ b/Engine/Source/Runtime/AnimGraphRuntime/Private/BoneControllers/AnimNode_RigidBody.cpp
@@ -1477,8 +1477,16 @@ void FAnimNode_RigidBody::PreUpdate(const UAnimInstance* InAnimInstance)
 			WorldTimeSeconds = SKC->PrimaryComponentTick.bTickEvenWhenPaused ? World->UnpausedTimeSeconds : World->TimeSeconds;
 
 			if (PhysicsSimulation && bEnableWorldGeometry)
-			{
+			{ 
+				// @todo: this logic can be simplified now. We used to run PurgeExpiredWorldObjects and CollectWorldObjects
+				// in UpdateAnimation, but we can't access the world actor's geometry there
 				UpdateWorldGeometry(*World, *SKC);
+
+				// Remove expired objects from the sim
+				PurgeExpiredWorldObjects();
+
+				// Find nearby world objects to add to the sim (gated on UnsafeWorld - see UpdateWorldGeometry)
+				CollectWorldObjects();
 			}
 
 			PendingRadialForces = SKC->GetPendingRadialForces();
@@ -1542,12 +1550,6 @@ void FAnimNode_RigidBody::UpdateInternal(const FAnimationUpdateContext& Context)
 		bSimulationStarted = true;
 	}
 
-	// Remove expired objects from the sim
-	PurgeExpiredWorldObjects();
-
-	// Find nearby world objects to add to the sim (gated on UnsafeWorld - see UpdateWorldGeometry)
-	CollectWorldObjects();
-
 	// These get set again if our bounds change. Subsequent calls to CollectWorldObjects will early-out until then
 	UnsafeWorld = nullptr;
 	UnsafeOwner = nullptr;
