Fixed a crash in the Pose Handler anim node, where it was using the bone container without first checking if it was valid

#rn fix animation/runtime

diff --git a/Engine/Source/Runtime/AnimGraphRuntime/Private/AnimNodes/AnimNode_PoseHandler.cpp b/Engine/Source/Runtime/AnimGraphRuntime/Private/AnimNodes/AnimNode_PoseHandler.cpp
--- a/Engine/Source/Runtime/AnimGraphRuntime/Private/AnimNodes/AnimNode_PoseHandler.cpp
+++ b/Engine/Source/Runtime/AnimGraphRuntime/Private/AnimNodes/AnimNode_PoseHandler.cpp
@@ -21,12 +21,13 @@ void FAnimNode_PoseHandler::CacheBoneBlendWeights(FAnimInstanceProxy* InstancePr
 {
 	BoneBlendWeights.Reset();
 
+	const FBoneContainer& BoneContainer = InstanceProxy->GetRequiredBones();
+
 	// this has to update bone blending weight
-	if (CurrentPoseAsset.IsValid())
+	if (CurrentPoseAsset.IsValid() && BoneContainer.IsValid())
 	{
 		const UPoseAsset* CurrentAsset = CurrentPoseAsset.Get();
 		const TArray<FName>& TrackNames = CurrentAsset->GetTrackNames();
-		const FBoneContainer& BoneContainer = InstanceProxy->GetRequiredBones();
 		const TArray<FBoneIndexType>& RequiredBoneIndices = BoneContainer.GetBoneIndicesArray();
 		BoneBlendWeights.AddZeroed(RequiredBoneIndices.Num());
 
