Fixed crashes caused by re-instancing clobbering the size of EditFixedSize arrays in anim nodes

Also patched up some other instances of anim nodes where this could occur

diff --git a/Engine/Source/Runtime/AnimGraphRuntime/Public/AnimNodes/AnimNode_LayeredBoneBlend.h b/Engine/Source/Runtime/AnimGraphRuntime/Public/AnimNodes/AnimNode_LayeredBoneBlend.h
--- a/Engine/Source/Runtime/AnimGraphRuntime/Public/AnimNodes/AnimNode_LayeredBoneBlend.h
+++ b/Engine/Source/Runtime/AnimGraphRuntime/Public/AnimNodes/AnimNode_LayeredBoneBlend.h
@@ -28,7 +28,7 @@ public:
 	FPoseLink BasePose;
 
 	/** Each layer's blended pose */
-	UPROPERTY(EditAnywhere, BlueprintReadWrite, editfixedsize, Category=Links)
+	UPROPERTY(EditAnywhere, BlueprintReadWrite, editfixedsize, Category=Links, meta=(BlueprintCompilerGeneratedDefaults))
 	TArray<FPoseLink> BlendPoses;
 
 	/** Whether to use branch filters or a blend mask to specify an input pose per-bone influence */
@@ -51,7 +51,7 @@ public:
 	TArray<FInputBlendPose> LayerSetup;
 
 	/** The weights of each layer */
-	UPROPERTY(EditAnywhere, BlueprintReadWrite, editfixedsize, Category=Runtime, meta=(PinShownByDefault))
+	UPROPERTY(EditAnywhere, BlueprintReadWrite, editfixedsize, Category=Runtime, meta=(BlueprintCompilerGeneratedDefaults, PinShownByDefault))
 	TArray<float> BlendWeights;
 
 	/** Whether to blend bone rotations in mesh space or in local space */
diff --git a/Engine/Source/Runtime/AnimGraphRuntime/Public/AnimNodes/AnimNode_ModifyCurve.h b/Engine/Source/Runtime/AnimGraphRuntime/Public/AnimNodes/AnimNode_ModifyCurve.h
--- a/Engine/Source/Runtime/AnimGraphRuntime/Public/AnimNodes/AnimNode_ModifyCurve.h
+++ b/Engine/Source/Runtime/AnimGraphRuntime/Public/AnimNodes/AnimNode_ModifyCurve.h
@@ -38,10 +38,10 @@ struct ANIMGRAPHRUNTIME_API FAnimNode_ModifyCurve : public FAnimNode_Base
 	UPROPERTY(EditAnywhere, BlueprintReadWrite, editfixedsize, Category = ModifyCurve, meta = (PinHiddenByDefault))
 	TMap<FName, float> CurveMap;
 
-	UPROPERTY(EditAnywhere, BlueprintReadWrite, editfixedsize, Category = ModifyCurve, meta = (PinShownByDefault))
+	UPROPERTY(EditAnywhere, BlueprintReadWrite, editfixedsize, Category = ModifyCurve, meta = (BlueprintCompilerGeneratedDefaults, PinShownByDefault))
 	TArray<float> CurveValues;
 
-	UPROPERTY()
+	UPROPERTY(meta = (BlueprintCompilerGeneratedDefaults))
 	TArray<FName> CurveNames;
 
 	TArray<float> LastCurveValues;
diff --git a/Engine/Source/Runtime/AnimGraphRuntime/Public/AnimNodes/AnimNode_MultiWayBlend.h b/Engine/Source/Runtime/AnimGraphRuntime/Public/AnimNodes/AnimNode_MultiWayBlend.h
--- a/Engine/Source/Runtime/AnimGraphRuntime/Public/AnimNodes/AnimNode_MultiWayBlend.h
+++ b/Engine/Source/Runtime/AnimGraphRuntime/Public/AnimNodes/AnimNode_MultiWayBlend.h
@@ -14,10 +14,10 @@ struct ANIMGRAPHRUNTIME_API FAnimNode_MultiWayBlend : public FAnimNode_Base
 
 public:
 
-	UPROPERTY(EditAnywhere, BlueprintReadWrite, Category = Links)
+	UPROPERTY(EditAnywhere, BlueprintReadWrite, Category = Links, meta = (BlueprintCompilerGeneratedDefaults))
 	TArray<FPoseLink> Poses;
 
-	UPROPERTY(EditAnywhere, EditFixedSize, BlueprintReadWrite, Category=Settings, meta=(PinShownByDefault))
+	UPROPERTY(EditAnywhere, EditFixedSize, BlueprintReadWrite, Category=Settings, meta=(BlueprintCompilerGeneratedDefaults, PinShownByDefault))
 	TArray<float> DesiredAlphas;
 
 private:
diff --git a/Engine/Source/Runtime/AnimGraphRuntime/Public/BoneControllers/AnimNode_SplineIK.h b/Engine/Source/Runtime/AnimGraphRuntime/Public/BoneControllers/AnimNode_SplineIK.h
--- a/Engine/Source/Runtime/AnimGraphRuntime/Public/BoneControllers/AnimNode_SplineIK.h
+++ b/Engine/Source/Runtime/AnimGraphRuntime/Public/BoneControllers/AnimNode_SplineIK.h
@@ -71,7 +71,7 @@ struct ANIMGRAPHRUNTIME_API FAnimNode_SplineIK : public FAnimNode_SkeletalContro
 	int32 PointCount;
 
 	/** Transforms applied to spline points **/
-	UPROPERTY(EditAnywhere, BlueprintReadWrite, EditFixedSize, Category = "Parameters", meta = (PinHiddenByDefault))
+	UPROPERTY(EditAnywhere, BlueprintReadWrite, EditFixedSize, Category = "Parameters", meta = (BlueprintCompilerGeneratedDefaults, PinHiddenByDefault))
 	TArray<FTransform> ControlPoints;
 
 	/** Overall roll of the spline, applied on top of other rotations along the direction of the spline */
