PR #8537: exposing composite anim sequences to random player (Contributed by [REDACTED_PERSON])

diff --git a/Engine/Source/Runtime/AnimGraphRuntime/Private/AnimNodes/AnimNode_RandomPlayer.cpp b/Engine/Source/Runtime/AnimGraphRuntime/Private/AnimNodes/AnimNode_RandomPlayer.cpp
--- a/Engine/Source/Runtime/AnimGraphRuntime/Private/AnimNodes/AnimNode_RandomPlayer.cpp
+++ b/Engine/Source/Runtime/AnimGraphRuntime/Private/AnimNodes/AnimNode_RandomPlayer.cpp
@@ -130,7 +130,7 @@ void FAnimNode_RandomPlayer::Update_AnyThread(const FAnimationUpdateContext& Con
 	FRandomAnimPlayData* CurrentData = &GetPlayData(ERandomDataIndexType::Current);
 	FRandomAnimPlayData* NextData = &GetPlayData(ERandomDataIndexType::Next);
 
-	const UAnimSequence* CurrentSequence = CurrentData->Entry->Sequence;
+	const UAnimSequenceBase* CurrentSequence = CurrentData->Entry->Sequence;
 
 	// If we looped around, adjust the previous play time to always be before the current playtime,
 	// since we can assume modulo. This makes the crossing check for the start time a lot simpler.
@@ -274,7 +274,7 @@ void FAnimNode_RandomPlayer::Evaluate_AnyThread(FPoseContext& Output)
 	FRandomAnimPlayData& CurrentData = GetPlayData(ERandomDataIndexType::Current);
 	FRandomAnimPlayData& NextData = GetPlayData(ERandomDataIndexType::Next);
 
-	UAnimSequence* CurrentSequence = CurrentData.Entry->Sequence;
+	UAnimSequenceBase* CurrentSequence = CurrentData.Entry->Sequence;
 
 	if (!FMath::IsNearlyEqualByULP(CurrentData.BlendWeight, 1.0f))
 	{
@@ -296,7 +296,7 @@ void FAnimNode_RandomPlayer::Evaluate_AnyThread(FPoseContext& Output)
 		Weights[0] = CurrentData.BlendWeight;
 		Weights[1] = NextData.BlendWeight;
 
-		UAnimSequence* NextSequence = NextData.Entry->Sequence;
+		UAnimSequenceBase* NextSequence = NextData.Entry->Sequence;
 
 		FAnimationPoseData CurrentPoseData(Poses[0], Curves[0], Attributes[0]);
 		FAnimationPoseData NextPoseData(Poses[1], Curves[1], Attributes[1]);
diff --git a/Engine/Source/Runtime/AnimGraphRuntime/Public/AnimNodes/AnimNode_RandomPlayer.h b/Engine/Source/Runtime/AnimGraphRuntime/Public/AnimNodes/AnimNode_RandomPlayer.h
--- a/Engine/Source/Runtime/AnimGraphRuntime/Public/AnimNodes/AnimNode_RandomPlayer.h
+++ b/Engine/Source/Runtime/AnimGraphRuntime/Public/AnimNodes/AnimNode_RandomPlayer.h
@@ -36,8 +36,8 @@ struct FRandomPlayerSequenceEntry
 	}
 
 	/** Sequence to play when this entry is picked */
-	UPROPERTY(EditAnywhere, BlueprintReadWrite, Category = "Settings")
-	TObjectPtr<UAnimSequence> Sequence;
+	UPROPERTY(EditAnywhere, BlueprintReadWrite, Category = "Settings", meta = (DisallowedClasses = "/Script/Engine.AnimMontage"))
+	TObjectPtr<UAnimSequenceBase> Sequence;
 
 	/** When not in shuffle mode, this is the chance this entry will play (normalized against all other sample chances) */
 	UPROPERTY(EditAnywhere, BlueprintReadWrite, Category = "Settings", meta = (UIMin = "0", ClampMin = "0"))
