PoseDriver: Remove ensure() if a curve doesn't exist. This can happen when using compatible skeletons that don't have any of the curves needed by the pose driver. The othe Pose* anim nodes behave the same way, silently ignoring missing curves.

#rnx

diff --git a/Engine/Source/Runtime/AnimGraphRuntime/Private/AnimNodes/AnimNode_PoseDriver.cpp b/Engine/Source/Runtime/AnimGraphRuntime/Private/AnimNodes/AnimNode_PoseDriver.cpp
--- a/Engine/Source/Runtime/AnimGraphRuntime/Private/AnimNodes/AnimNode_PoseDriver.cpp
+++ b/Engine/Source/Runtime/AnimGraphRuntime/Private/AnimNodes/AnimNode_PoseDriver.cpp
@@ -58,7 +58,7 @@ void FAnimNode_PoseDriver::RebuildPoseList(const FBoneContainer& InBoneContainer
 			if (PoseIndex != INDEX_NONE)
 			{
 				TArray<uint16> const& LUTIndex = InBoneContainer.GetUIDToArrayLookupTable();
-				if (ensure(LUTIndex.IsValidIndex(PoseNames[PoseIndex].UID)) && LUTIndex[PoseNames[PoseIndex].UID] != MAX_uint16)
+				if (LUTIndex.IsValidIndex(PoseNames[PoseIndex].UID) && LUTIndex[PoseNames[PoseIndex].UID] != MAX_uint16)
 				{
 					// we keep pose index as that is the fastest way to search when extracting pose asset
 					PoseTarget.PoseCurveIndex = PoseExtractContext.PoseCurves.Add(FPoseCurve(PoseIndex, PoseNames[PoseIndex].UID, 0.f));
