Truncation fixes for AnimationGraphRuntime module

diff --git a/Engine/Source/Editor/AnimGraph/Private/EditModes/PoseDriverEditMode.cpp b/Engine/Source/Editor/AnimGraph/Private/EditModes/PoseDriverEditMode.cpp
--- a/Engine/Source/Editor/AnimGraph/Private/EditModes/PoseDriverEditMode.cpp
+++ b/Engine/Source/Editor/AnimGraph/Private/EditModes/PoseDriverEditMode.cpp
@@ -80,9 +80,9 @@ void FPoseDriverEditMode::Render(const FSceneView* View, FViewport* Viewport, FP
 
 	static const float DrawLineWidth = 0.1f;
 	static const float DrawPosSize = 2.f;
-	float DrawAxisLength = GraphNode->AxisLength;
-	int32 DrawConeSubdivision = GraphNode->ConeSubdivision;
-	bool bDrawDebugCones = GraphNode->bDrawDebugCones;
+	const float DrawAxisLength = GraphNode->AxisLength;
+	const int32 DrawConeSubdivision = GraphNode->ConeSubdivision;
+	const bool bDrawDebugCones = GraphNode->bDrawDebugCones;
 
 
 	if (SkelComp->AnimScriptInstance)
@@ -102,16 +102,16 @@ void FPoseDriverEditMode::Render(const FSceneView* View, FViewport* Viewport, FP
 			if (BoneIndex != INDEX_NONE)
 			{
 				// Get transform of driven bone, used as basis for drawing
-				FTransform BoneWorldTM = SkelComp->GetBoneTransform(BoneIndex);
-				FVector BonePos = BoneWorldTM.GetLocation();
+				const FTransform BoneWorldTM = SkelComp->GetBoneTransform(BoneIndex);
+				const FVector BonePos = BoneWorldTM.GetLocation();
 
 				// Transform that we are evaluating pose in
 				FTransform EvalSpaceTM = SkelComp->GetComponentToWorld();
 
 				// If specifying space to eval in, get that space
-				int32 EvalSpaceBoneIndex = SkelComp->GetBoneIndex(RuntimeNode->EvalSpaceBone.BoneName);
-				FName ParentBoneName = SkelComp->GetParentBone(SourceBoneRef.BoneName);
-				int32 ParentBoneIndex = SkelComp->GetBoneIndex(ParentBoneName);
+				const int32 EvalSpaceBoneIndex = SkelComp->GetBoneIndex(RuntimeNode->EvalSpaceBone.BoneName);
+				const FName ParentBoneName = SkelComp->GetParentBone(SourceBoneRef.BoneName);
+				const int32 ParentBoneIndex = SkelComp->GetBoneIndex(ParentBoneName);
 				if (EvalSpaceBoneIndex != INDEX_NONE)
 				{
 					EvalSpaceTM = SkelComp->GetBoneTransform(EvalSpaceBoneIndex);
@@ -125,7 +125,7 @@ void FPoseDriverEditMode::Render(const FSceneView* View, FViewport* Viewport, FP
 				// Get source bone TM from last frame
 				if (RuntimeNode->SourceBoneTMs.IsValidIndex(SourceIdx))
 				{
-					FTransform SourceBoneTM = RuntimeNode->SourceBoneTMs[SourceIdx];
+					const FTransform SourceBoneTM = RuntimeNode->SourceBoneTMs[SourceIdx];
 
 					// Rotation drawing
 					if (RuntimeNode->DriveSource == EPoseDriverSource::Rotation)
@@ -143,21 +143,21 @@ void FPoseDriverEditMode::Render(const FSceneView* View, FViewport* Viewport, FP
 								LocalTwistVec = FRotator(MedianRot.X, MedianRot.Y, MedianRot.Z).RotateVector(RuntimeNode->RBFParams.GetTwistAxisVector());
 								WorldTwistVec = EvalSpaceTM.TransformVectorNoScale(LocalTwistVec);
 
-								FVector LocalSwingVec = FVector::CrossProduct(LocalTwistVec, FVector(1, 1, 1));
+								const FVector LocalSwingVec = FVector::CrossProduct(LocalTwistVec, FVector(1, 1, 1));
 								FVector WorldSwingVec = EvalSpaceTM.TransformVectorNoScale(LocalSwingVec);
 								WorldSwingVec.Normalize();
 
-								FQuat WorldTwistQ(WorldTwistVec, PI * 2.0f / float(DrawConeSubdivision));
-								FQuat WorldSwingMinQ(WorldSwingVec, FMath::DegreesToRadians(RuntimeNode->RBFParams.MedianMin));
-								FQuat WorldSwingMaxQ(WorldSwingVec, FMath::DegreesToRadians(RuntimeNode->RBFParams.MedianMax));
+								const FQuat WorldTwistQ(WorldTwistVec, UE_PI * 2.0f / float(DrawConeSubdivision));
+								const FQuat WorldSwingMinQ(WorldSwingVec, FMath::DegreesToRadians(RuntimeNode->RBFParams.MedianMin));
+								const FQuat WorldSwingMaxQ(WorldSwingVec, FMath::DegreesToRadians(RuntimeNode->RBFParams.MedianMax));
 
-								FVector FirstMinPositionOnCircle = WorldSwingMinQ.RotateVector(WorldTwistVec * DrawAxisLength);
-								FVector FirstMaxPositionOnCircle = WorldSwingMaxQ.RotateVector(WorldTwistVec * DrawAxisLength);
+								const FVector FirstMinPositionOnCircle = WorldSwingMinQ.RotateVector(WorldTwistVec * DrawAxisLength);
+								const FVector FirstMaxPositionOnCircle = WorldSwingMaxQ.RotateVector(WorldTwistVec * DrawAxisLength);
 								FVector LastMinPositionOnCircle = FirstMinPositionOnCircle;
 								FVector LastMaxPositionOnCircle = FirstMaxPositionOnCircle;
 
-								FLinearColor MinColor = FLinearColor::Yellow;
-								FLinearColor MaxColor = MinColor.Desaturate(0.5);
+								const FLinearColor MinColor = FLinearColor::Yellow;
+								const FLinearColor MaxColor = MinColor.Desaturate(0.5);
 
 								for (int32 i = 0; i < DrawConeSubdivision; i++)
 								{
@@ -180,7 +180,7 @@ void FPoseDriverEditMode::Render(const FSceneView* View, FViewport* Viewport, FP
 					// Translation drawing
 					else if (RuntimeNode->DriveSource == EPoseDriverSource::Translation)
 					{
-						FVector LocalPos = SourceBoneTM.GetTranslation();
+						const FVector LocalPos = SourceBoneTM.GetTranslation();
 						FVector WorldPos = EvalSpaceTM.TransformPosition(LocalPos);
 						DrawWireDiamond(PDI, FTranslationMatrix(WorldPos), DrawPosSize, FLinearColor::Green, SDPG_Foreground);
 
@@ -224,34 +224,34 @@ void FPoseDriverEditMode::Render(const FSceneView* View, FViewport* Viewport, FP
 						{
 							const FPoseDriverTransform& TargetTM = PoseTarget.BoneTransforms[SourceIdx];
 
-							bool bSelected = (GraphNode->SelectedTargetIndex == TargetIdx);
-							float AxisLength = bSelected ? (DrawAxisLength * 1.5f) : DrawAxisLength;
-							float LineWidth = bSelected ? (DrawLineWidth * 3.f) : DrawLineWidth;
-							float Radius = RuntimeNode->GetRadiusForTarget(RBFTarget);
+							const bool bSelected = (GraphNode->SelectedTargetIndex == TargetIdx);
+							const float AxisLength = bSelected ? (DrawAxisLength * 1.5f) : DrawAxisLength;
+							const float LineWidth = bSelected ? (DrawLineWidth * 3.f) : DrawLineWidth;
+							const float Radius = RuntimeNode->GetRadiusForTarget(RBFTarget);
 							
-							FLinearColor Color = TotalWeight <= 0.f ? FLinearColor::Black : GraphNode->GetColorFromWeight(PerTargetWeights[TargetIdx]);
+							const FLinearColor Color = TotalWeight <= 0.f ? FLinearColor::Black : GraphNode->GetColorFromWeight(PerTargetWeights[TargetIdx]);
 
 							PDI->SetHitProxy(new HPDTargetHitProxy(TargetIdx));
 
 							// Rotation drawing
 							if (RuntimeNode->DriveSource == EPoseDriverSource::Rotation)
 							{
-								FVector LocalTwistVec = TargetTM.TargetRotation.RotateVector(RuntimeNode->RBFParams.GetTwistAxisVector());
-								FVector WorldTwistVec = EvalSpaceTM.TransformVectorNoScale(LocalTwistVec);
+								const FVector LocalTwistVec = TargetTM.TargetRotation.RotateVector(RuntimeNode->RBFParams.GetTwistAxisVector());
+								const FVector WorldTwistVec = EvalSpaceTM.TransformVectorNoScale(LocalTwistVec);
 
-								FVector LocalSwingVec = FVector::CrossProduct(LocalTwistVec, FVector(1, 1, 1));
+								const FVector LocalSwingVec = FVector::CrossProduct(LocalTwistVec, FVector(1, 1, 1));
 								FVector WorldSwingVec = EvalSpaceTM.TransformVectorNoScale(LocalSwingVec);
 								WorldSwingVec.Normalize();
 
-								FQuat WorldTwistQ(WorldTwistVec, PI * 2.0f / float(DrawConeSubdivision));
-								FQuat WorldSwingQ(WorldSwingVec, FMath::DegreesToRadians(Radius));
+								const FQuat WorldTwistQ(WorldTwistVec, UE_PI * 2.0f / float(DrawConeSubdivision));
+								const FQuat WorldSwingQ(WorldSwingVec, FMath::DegreesToRadians(Radius));
 
-								FVector FirstPositionOnCircle = WorldSwingQ.RotateVector(WorldTwistVec * DrawAxisLength);
+								const FVector FirstPositionOnCircle = WorldSwingQ.RotateVector(WorldTwistVec * DrawAxisLength);
 								FVector LastPositionOnCircle = FirstPositionOnCircle;
 
 								for (int32 i = 0; i < DrawConeSubdivision; i++)
 								{
-									FVector NextPositionOnCircle = WorldTwistQ.RotateVector(LastPositionOnCircle);
+									const FVector NextPositionOnCircle = WorldTwistQ.RotateVector(LastPositionOnCircle);
 									PDI->DrawLine(BonePos, BonePos + NextPositionOnCircle, Color, SDPG_Foreground, DrawLineWidth);
 									PDI->DrawLine(BonePos + LastPositionOnCircle, BonePos + NextPositionOnCircle, Color, SDPG_Foreground, DrawLineWidth);
 
@@ -262,8 +262,8 @@ void FPoseDriverEditMode::Render(const FSceneView* View, FViewport* Viewport, FP
 							// Translation drawing
 							else if (RuntimeNode->DriveSource == EPoseDriverSource::Translation)
 							{
-								FVector LocalPos = TargetTM.TargetTranslation;
-								FVector WorldPos = EvalSpaceTM.TransformPosition(LocalPos);
+								const FVector LocalPos = TargetTM.TargetTranslation;
+								const FVector WorldPos = EvalSpaceTM.TransformPosition(LocalPos);
 								DrawWireDiamond(PDI, FTranslationMatrix(WorldPos), Radius, Color, SDPG_Foreground, LineWidth);
 							}
 
@@ -282,7 +282,7 @@ bool FPoseDriverEditMode::HandleClick(FEditorViewportClient* InViewportClient, H
 
 	if (HitProxy != nullptr && HitProxy->IsA(HPDTargetHitProxy::StaticGetType()))
 	{
-		HPDTargetHitProxy* TargetHitProxy = static_cast<HPDTargetHitProxy*>(HitProxy);
+		const HPDTargetHitProxy* TargetHitProxy = static_cast<HPDTargetHitProxy*>(HitProxy);
 		GraphNode->SelectedTargetIndex = TargetHitProxy->TargetIndex;
 		GraphNode->SelectedTargetChangeDelegate.Broadcast();
 		bResult = true;
diff --git a/Engine/Source/Runtime/AnimGraphRuntime/Private/AnimGraphRuntimeTrace.cpp b/Engine/Source/Runtime/AnimGraphRuntime/Private/AnimGraphRuntimeTrace.cpp
--- a/Engine/Source/Runtime/AnimGraphRuntime/Private/AnimGraphRuntimeTrace.cpp
+++ b/Engine/Source/Runtime/AnimGraphRuntime/Private/AnimGraphRuntimeTrace.cpp
@@ -44,12 +44,12 @@ void FAnimGraphRuntimeTrace::OutputBlendSpacePlayer(const FAnimationBaseContext&
 		<< BlendSpacePlayer.AnimInstanceId(FObjectTrace::GetObjectId(InContext.AnimInstanceProxy->GetAnimInstanceObject()))
 		<< BlendSpacePlayer.BlendSpaceId(FObjectTrace::GetObjectId(InNode.GetBlendSpace()))
 		<< BlendSpacePlayer.NodeId(InContext.GetCurrentNodeId())
-		<< BlendSpacePlayer.PositionX(SampleCoordinates.X)
-		<< BlendSpacePlayer.PositionY(SampleCoordinates.Y)
-		<< BlendSpacePlayer.PositionZ(SampleCoordinates.Z)
-		<< BlendSpacePlayer.FilteredPositionX(FilteredPosition.X)
-		<< BlendSpacePlayer.FilteredPositionY(FilteredPosition.Y)
-		<< BlendSpacePlayer.FilteredPositionZ(FilteredPosition.Z);
+		<< BlendSpacePlayer.PositionX(static_cast<float>(SampleCoordinates.X))
+		<< BlendSpacePlayer.PositionY(static_cast<float>(SampleCoordinates.Y))
+		<< BlendSpacePlayer.PositionZ(static_cast<float>(SampleCoordinates.Z))
+		<< BlendSpacePlayer.FilteredPositionX(static_cast<float>(FilteredPosition.X))
+		<< BlendSpacePlayer.FilteredPositionY(static_cast<float>(FilteredPosition.Y))
+		<< BlendSpacePlayer.FilteredPositionZ(static_cast<float>(FilteredPosition.Z));
 }
 
 void FAnimGraphRuntimeTrace::OutputBlendSpace(const FAnimationBaseContext& InContext, const FAnimNode_BlendSpaceGraphBase& InNode)
@@ -73,12 +73,12 @@ void FAnimGraphRuntimeTrace::OutputBlendSpace(const FAnimationBaseContext& InCon
 		<< BlendSpacePlayer.AnimInstanceId(FObjectTrace::GetObjectId(InContext.AnimInstanceProxy->GetAnimInstanceObject()))
 		<< BlendSpacePlayer.BlendSpaceId(FObjectTrace::GetObjectId(InNode.GetBlendSpace()))
 		<< BlendSpacePlayer.NodeId(InContext.GetCurrentNodeId())
-		<< BlendSpacePlayer.PositionX(Coordinates.X)
-		<< BlendSpacePlayer.PositionY(Coordinates.Y)
-		<< BlendSpacePlayer.PositionZ(Coordinates.Z)
-		<< BlendSpacePlayer.FilteredPositionX(FilteredCoordinates.X)
-		<< BlendSpacePlayer.FilteredPositionY(FilteredCoordinates.Y)
-		<< BlendSpacePlayer.FilteredPositionZ(FilteredCoordinates.Z);
+		<< BlendSpacePlayer.PositionX(static_cast<float>(Coordinates.X))
+		<< BlendSpacePlayer.PositionY(static_cast<float>(Coordinates.Y))
+		<< BlendSpacePlayer.PositionZ(static_cast<float>(Coordinates.Z))
+		<< BlendSpacePlayer.FilteredPositionX(static_cast<float>(FilteredCoordinates.X))
+		<< BlendSpacePlayer.FilteredPositionY(static_cast<float>(FilteredCoordinates.Y))
+		<< BlendSpacePlayer.FilteredPositionZ(static_cast<float>(FilteredCoordinates.Z));
 }
 
 #endif
\ No newline at end of file
diff --git a/Engine/Source/Runtime/AnimGraphRuntime/Private/AnimNodes/AnimNode_BlendSpacePlayer.cpp b/Engine/Source/Runtime/AnimGraphRuntime/Private/AnimNodes/AnimNode_BlendSpacePlayer.cpp
--- a/Engine/Source/Runtime/AnimGraphRuntime/Private/AnimNodes/AnimNode_BlendSpacePlayer.cpp
+++ b/Engine/Source/Runtime/AnimGraphRuntime/Private/AnimNodes/AnimNode_BlendSpacePlayer.cpp
@@ -162,7 +162,7 @@ void FAnimNode_BlendSpacePlayerBase::SnapToPosition(const FVector& NewPosition)
 	const int32 NumAxis = FMath::Min(BlendFilter.FilterPerAxis.Num(), 3);
 	for (int32 Idx = 0; Idx < NumAxis; Idx++)
 	{
-		BlendFilter.FilterPerAxis[Idx].SetToValue(NewPosition[Idx]);
+		BlendFilter.FilterPerAxis[Idx].SetToValue(static_cast<float>(NewPosition[Idx]));
 	}
 }
 
@@ -325,10 +325,10 @@ FVector FAnimNode_BlendSpacePlayer::GetPosition() const
 bool FAnimNode_BlendSpacePlayer::SetPosition(FVector InPosition)
 {
 #if WITH_EDITORONLY_DATA
-	X = InPosition[0];
-	Y = InPosition[1];
-	GET_MUTABLE_ANIM_NODE_DATA(float, X) = InPosition[0];
-	GET_MUTABLE_ANIM_NODE_DATA(float, Y) = InPosition[1];
+	X = static_cast<float>(InPosition[0]);
+	Y = static_cast<float>(InPosition[1]);
+	GET_MUTABLE_ANIM_NODE_DATA(float, X) = static_cast<float>(InPosition[0]);
+	GET_MUTABLE_ANIM_NODE_DATA(float, Y) = static_cast<float>(InPosition[1]);
 #endif
 
 	float* XPtr = GET_INSTANCE_ANIM_NODE_DATA_PTR(float, X);
@@ -336,8 +336,8 @@ bool FAnimNode_BlendSpacePlayer::SetPosition(FVector InPosition)
 
 	if (XPtr && YPtr)
 	{
-		*XPtr = InPosition[0];
-		*YPtr = InPosition[1];
+		*XPtr = static_cast<float>(InPosition[0]);
+		*YPtr = static_cast<float>(InPosition[1]);
 		return true;
 	}
 
diff --git a/Engine/Source/Runtime/AnimGraphRuntime/Private/AnimNodes/AnimNode_LayeredBoneBlend.cpp b/Engine/Source/Runtime/AnimGraphRuntime/Private/AnimNodes/AnimNode_LayeredBoneBlend.cpp
--- a/Engine/Source/Runtime/AnimGraphRuntime/Private/AnimNodes/AnimNode_LayeredBoneBlend.cpp
+++ b/Engine/Source/Runtime/AnimGraphRuntime/Private/AnimNodes/AnimNode_LayeredBoneBlend.cpp
@@ -113,7 +113,7 @@ void FAnimNode_LayeredBoneBlend::UpdateCachedBoneData(const FBoneContainer& Requ
 						{
 							if (DesiredBoneBlendWeights[CompactPoseIndex.GetInt()].BlendWeight > 0.f)
 							{
-								CurvePoseSourceIndices[CurrentPoseIndex] = DesiredBoneBlendWeights[CompactPoseIndex.GetInt()].SourceIndex;
+								CurvePoseSourceIndices[CurrentPoseIndex] = IntCastChecked<uint8>(DesiredBoneBlendWeights[CompactPoseIndex.GetInt()].SourceIndex);
 							}
 						}
 					}
@@ -248,7 +248,7 @@ void FAnimNode_LayeredBoneBlend::Evaluate_AnyThread(FPoseContext& Output)
 		TArray<uint16> const* CurveUIDFinder = Output.Curve.UIDToArrayIndexLUT;
 		const int32 TotalCount = Output.Curve.NumValidCurveCount;
 		// now go through point to correct source indices. Curve only picks one source index
-		for (int32 UIDIndex = 0; UIDIndex < CurveUIDFinder->Num(); ++UIDIndex)
+		for (USkeleton::AnimCurveUID UIDIndex = 0; UIDIndex < CurveUIDFinder->Num(); ++UIDIndex)
 		{
 			int32 CurvePoseIndex = Output.Curve.GetArrayIndexByUID(UIDIndex);
 			if (CurvePoseSourceIndices.IsValidIndex(CurvePoseIndex))
diff --git a/Engine/Source/Runtime/AnimGraphRuntime/Private/AnimNodes/AnimNode_PoseDriver.cpp b/Engine/Source/Runtime/AnimGraphRuntime/Private/AnimNodes/AnimNode_PoseDriver.cpp
--- a/Engine/Source/Runtime/AnimGraphRuntime/Private/AnimNodes/AnimNode_PoseDriver.cpp
+++ b/Engine/Source/Runtime/AnimGraphRuntime/Private/AnimNodes/AnimNode_PoseDriver.cpp
@@ -51,7 +51,7 @@ void FAnimNode_PoseDriver::RebuildPoseList(const FBoneContainer& InBoneContainer
 			}
 			else
 			{
-				PoseTarget.DrivenUID = INDEX_NONE;
+				PoseTarget.DrivenUID = SmartName::MaxUID;
 			}
 
 			const int32 PoseIndex = InPoseAsset->GetPoseIndexByName(PoseTarget.DrivenName);
diff --git a/Engine/Source/Runtime/AnimGraphRuntime/Private/BoneControllers/AnimNode_AnimDynamics.cpp b/Engine/Source/Runtime/AnimGraphRuntime/Private/BoneControllers/AnimNode_AnimDynamics.cpp
--- a/Engine/Source/Runtime/AnimGraphRuntime/Private/BoneControllers/AnimNode_AnimDynamics.cpp
+++ b/Engine/Source/Runtime/AnimGraphRuntime/Private/BoneControllers/AnimNode_AnimDynamics.cpp
@@ -721,13 +721,13 @@ void FAnimNode_AnimDynamics::InitPhysics(FComponentSpacePoseContext& Output)
 			switch (PhysicsBody.CollisionType)
 			{
 			case AnimPhysCollisionType::CustomSphere:
-				PhysicsBody.SphereCollisionRadius = PhysicsBodyDef.SphereCollisionRadius;
+				PhysicsBody.SphereCollisionRadius = static_cast<float>(PhysicsBodyDef.SphereCollisionRadius);
 				break;
 			case AnimPhysCollisionType::InnerSphere:
-				PhysicsBody.SphereCollisionRadius = PhysicsBodyDef.BoxExtents.GetAbsMin() / 2.0f;
+				PhysicsBody.SphereCollisionRadius = static_cast<float>(PhysicsBodyDef.BoxExtents.GetAbsMin() / 2.0);
 				break;
 			case AnimPhysCollisionType::OuterSphere:
-				PhysicsBody.SphereCollisionRadius = PhysicsBodyDef.BoxExtents.GetAbsMax() / 2.0f;
+				PhysicsBody.SphereCollisionRadius = static_cast<float>(PhysicsBodyDef.BoxExtents.GetAbsMax() / 2.0);
 				break;
 			default:
 				break;
diff --git a/Engine/Source/Runtime/AnimGraphRuntime/Private/BoneControllers/AnimNode_BoneDrivenController.cpp b/Engine/Source/Runtime/AnimGraphRuntime/Private/BoneControllers/AnimNode_BoneDrivenController.cpp
--- a/Engine/Source/Runtime/AnimGraphRuntime/Private/BoneControllers/AnimNode_BoneDrivenController.cpp
+++ b/Engine/Source/Runtime/AnimGraphRuntime/Private/BoneControllers/AnimNode_BoneDrivenController.cpp
@@ -73,7 +73,7 @@ void FAnimNode_BoneDrivenController::EvaluateSkeletalControl_AnyThread(FComponen
 	const FTransform& SourceRefPoseBoneTransform = BoneContainer.GetRefPoseArray()[SourceBone.BoneIndex];
 	const FTransform SourceCurrentBoneTransform = Output.Pose.GetLocalSpaceTransform(SourceBone.GetCompactPoseIndex(BoneContainer));
 
-	const float FinalDriverValue = ExtractSourceValue(SourceCurrentBoneTransform, SourceRefPoseBoneTransform);
+	const double FinalDriverValue = ExtractSourceValue(SourceCurrentBoneTransform, SourceRefPoseBoneTransform);
 	
 	
 	// Calculate a new local-space bone position by adding or replacing target components in the current local space position
@@ -200,24 +200,24 @@ void FAnimNode_BoneDrivenController::EvaluateComponentSpaceInternal(FComponentSp
 	const FTransform& SourceRefPoseBoneTransform = BoneContainer.GetRefPoseArray()[SourceBone.BoneIndex];
 	const FTransform SourceCurrentBoneTransform = Context.Pose.GetLocalSpaceTransform(SourceBone.GetCompactPoseIndex(BoneContainer));
 
-	const float FinalDriverValue = ExtractSourceValue(SourceCurrentBoneTransform, SourceRefPoseBoneTransform);
+	const double FinalDriverValue = ExtractSourceValue(SourceCurrentBoneTransform, SourceRefPoseBoneTransform);
 
 	if (DestinationMode == EDrivenDestinationMode::MorphTarget || DestinationMode == EDrivenDestinationMode::MaterialParameter)
 	{
 		//	Morph target and Material parameter curves
 		USkeleton* Skeleton = Context.AnimInstanceProxy->GetSkeleton();
-		SmartName::UID_Type NameUID = Skeleton->GetUIDByName(USkeleton::AnimCurveMappingName, ParameterName);
+		USkeleton::AnimCurveUID NameUID = Skeleton->GetUIDByName(USkeleton::AnimCurveMappingName, ParameterName);
 		if (NameUID != SmartName::MaxUID)
 		{
-			Context.Curve.Set(NameUID, FinalDriverValue);
+			Context.Curve.Set(NameUID, static_cast<float>(FinalDriverValue));
 		}
 	}
 }
 
-const float FAnimNode_BoneDrivenController::ExtractSourceValue(const FTransform &InCurrentBoneTransform, const FTransform &InRefPoseBoneTransform)
+const double FAnimNode_BoneDrivenController::ExtractSourceValue(const FTransform &InCurrentBoneTransform, const FTransform &InRefPoseBoneTransform)
 {
 	// Resolve source value
-	float SourceValue = 0.0f;
+	double SourceValue = 0.0;
 	if (SourceComponent < EComponentType::RotationX)
 	{
 		const FVector TranslationDiff = InCurrentBoneTransform.GetLocation() - InRefPoseBoneTransform.GetLocation();
@@ -232,7 +232,7 @@ const float FAnimNode_BoneDrivenController::ExtractSourceValue(const FTransform
 	{
 		const FVector CurrentScale = InCurrentBoneTransform.GetScale3D();
 		const FVector RefScale = InRefPoseBoneTransform.GetScale3D();
-		const float ScaleDiff = FMath::Max3(CurrentScale[0], CurrentScale[1], CurrentScale[2]) - FMath::Max3(RefScale[0], RefScale[1], RefScale[2]);
+		const double ScaleDiff = FMath::Max3(CurrentScale[0], CurrentScale[1], CurrentScale[2]) - FMath::Max3(RefScale[0], RefScale[1], RefScale[2]);
 		SourceValue = ScaleDiff;
 	}
 	else
@@ -242,18 +242,18 @@ const float FAnimNode_BoneDrivenController::ExtractSourceValue(const FTransform
 	}
 
 	// Determine the resulting value
-	float FinalDriverValue = SourceValue;
+	double FinalDriverValue = SourceValue;
 	if (DrivingCurve != nullptr)
 	{
 		// Remap thru the curve if set
-		FinalDriverValue = DrivingCurve->GetFloatValue(FinalDriverValue);
+		FinalDriverValue = DrivingCurve->GetFloatValue(static_cast<float>(FinalDriverValue));
 	}
 	else
 	{
 		// Apply the fixed function remapping/clamping
 		if (bUseRange)
 		{
-			const float ClampedAlpha = FMath::Clamp(FMath::GetRangePct(RangeMin, RangeMax, FinalDriverValue), 0.0f, 1.0f);
+			const double ClampedAlpha = FMath::Clamp(FMath::GetRangePct(RangeMin, RangeMax, FinalDriverValue), 0.0, 1.0);
 			FinalDriverValue = FMath::Lerp(RemappedMin, RemappedMax, ClampedAlpha);
 		}
 
diff --git a/Engine/Source/Runtime/AnimGraphRuntime/Private/BoneControllers/AnimNode_Fabrik.cpp b/Engine/Source/Runtime/AnimGraphRuntime/Private/BoneControllers/AnimNode_Fabrik.cpp
--- a/Engine/Source/Runtime/AnimGraphRuntime/Private/BoneControllers/AnimNode_Fabrik.cpp
+++ b/Engine/Source/Runtime/AnimGraphRuntime/Private/BoneControllers/AnimNode_Fabrik.cpp
@@ -79,7 +79,7 @@ void FAnimNode_Fabrik::EvaluateSkeletalControl_AnyThread(FComponentSpacePoseCont
 	}
 
 	// Maximum length of skeleton segment at full extension
-	float MaximumReach = 0;
+	double MaximumReach = 0;
 
 	// Gather transforms
 	int32 const NumTransforms = BoneIndices.Num();
@@ -109,7 +109,7 @@ void FAnimNode_Fabrik::EvaluateSkeletalControl_AnyThread(FComponentSpacePoseCont
 		OutBoneTransforms[TransformIndex] = FBoneTransform(BoneIndex, BoneCSTransform);
 
 		// Calculate the combined length of this segment of skeleton
-		float const BoneLength = FVector::Dist(BoneCSPosition, OutBoneTransforms[TransformIndex-1].Transform.GetLocation());
+		double const BoneLength = FVector::Dist(BoneCSPosition, OutBoneTransforms[TransformIndex-1].Transform.GetLocation());
 
 		if (!FMath::IsNearlyZero(BoneLength))
 		{
@@ -158,7 +158,7 @@ void FAnimNode_Fabrik::EvaluateSkeletalControl_AnyThread(FComponentSpacePoseCont
 
 			// Calculate axis of rotation from pre-translation vector to post-translation vector
 			FVector const RotationAxis = FVector::CrossProduct(OldDir, NewDir).GetSafeNormal();
-			float const RotationAngle = FMath::Acos(FVector::DotProduct(OldDir, NewDir));
+			double const RotationAngle = FMath::Acos(FVector::DotProduct(OldDir, NewDir));
 			FQuat const DeltaRotation = FQuat(RotationAxis, RotationAngle);
 			// We're going to multiply it, in order to not have to re-normalize the final quaternion, it has to be a unit quaternion.
 			checkSlow(DeltaRotation.IsNormalized());
diff --git a/Engine/Source/Runtime/AnimGraphRuntime/Private/BoneControllers/AnimNode_LegIK.cpp b/Engine/Source/Runtime/AnimGraphRuntime/Private/BoneControllers/AnimNode_LegIK.cpp
--- a/Engine/Source/Runtime/AnimGraphRuntime/Private/BoneControllers/AnimNode_LegIK.cpp
+++ b/Engine/Source/Runtime/AnimGraphRuntime/Private/BoneControllers/AnimNode_LegIK.cpp
@@ -216,14 +216,14 @@ void FIKChain::InitializeFromLegData(FAnimLegIKData& InLegData, FAnimInstancePro
 		Links.Init(FIKChainLink(), InLegData.NumBones);
 	}
 	
-	MaximumReach = 0.f;
+	MaximumReach = 0.0;
 
 	check(InLegData.NumBones > 1);
 	for (int32 Index = 0; Index < InLegData.NumBones - 1; Index++)
 	{
 		const FVector BoneLocation = InLegData.FKLegBoneTransforms[Index].GetLocation();
 		const FVector ParentLocation = InLegData.FKLegBoneTransforms[Index + 1].GetLocation();
-		const float BoneLength = FVector::Dist(BoneLocation, ParentLocation);
+		const double BoneLength = FVector::Dist(BoneLocation, ParentLocation);
 
 		FIKChainLink& Link = Links[Index];
 		Link.Location = BoneLocation;
@@ -259,7 +259,7 @@ void FIKChain::InitializeFromLegData(FAnimLegIKData& InLegData, FAnimInstancePro
 
 TAutoConsoleVariable<int32> CVarAnimLegIKTwoBone(TEXT("a.AnimNode.LegIK.EnableTwoBone"), 1, TEXT("Enable Two Bone Code Path."));
 
-void FIKChain::ReachTarget(const FVector& InTargetLocation, float InReachPrecision, int32 InMaxIterations)
+void FIKChain::ReachTarget(const FVector& InTargetLocation, double InReachPrecision, int32 InMaxIterations)
 {
 	if (!bInitialized)
 	{
@@ -309,13 +309,13 @@ void FIKChain::SolveTwoBoneIK(const FVector& InTargetLocation)
 
 	// Use Law of Cosines to work out solution.
 	// At this point we know the target location is reachable, and we are already aligned with that location. So the leg is in the right plane.
-	const float a = Links[1].Length;	// hip to knee
-	const float b = HipToFoot.Size();	// hip to foot
-	const float c = Links[0].Length;	// knee to foot
+	const double a = Links[1].Length;	// hip to knee
+	const double b = HipToFoot.Size();	// hip to foot
+	const double c = Links[0].Length;	// knee to foot
 
-	const float Two_ab = 2.f * a * b;
-	const float CosC = !FMath::IsNearlyZero(Two_ab) ? (a * a + b * b - c * c) / Two_ab : 0.f;
- 	const float C = FMath::Acos(CosC);
+	const double Two_ab = 2.f * a * b;
+	const double CosC = !FMath::IsNearlyZero(Two_ab) ? (a * a + b * b - c * c) / Two_ab : 0.0;
+ 	const double C = FMath::Acos(CosC);
 	
 	// Project Knee onto Hip to Foot line.
 	const FVector HipToFootDir = !FMath::IsNearlyZero(b) ? HipToFoot / b : FVector::ZeroVector;
@@ -332,13 +332,13 @@ void FIKChain::SolveTwoBoneIK(const FVector& InTargetLocation)
 	if ((HingeRotationAxis != FVector::ZeroVector) && (HipToFootDir != FVector::ZeroVector) && !FMath::IsNearlyZero(a))
 	{
 		const FVector HipToKneeDir = HipToKnee / a;
-		const float KneeBendDot = HipToKneeDir | HipToFootDir;
+		const double KneeBendDot = HipToKneeDir | HipToFootDir;
 
 		FVector& CachedRealBendDir = Links[1].RealBendDir;
 		FVector& CachedBaseBendDir = Links[1].BaseBendDir;
 
 		// Valid 'bend', cache 'BendDir'
-		if ((BendDir != FVector::ZeroVector) && (KneeBendDot < 0.99f))
+		if ((BendDir != FVector::ZeroVector) && (KneeBendDot < 0.99))
 		{
 			CachedRealBendDir = BendDir;
 			CachedBaseBendDir = HingeRotationAxis ^ HipToFootDir;
@@ -456,10 +456,10 @@ void FIKChain::FABRIK_ApplyLinkConstraints_Forward(FIKChain& IKChain, int32 Link
 	const FVector ChildAxisY = CurrentLink.LinkAxisZ ^ ChildAxisX;
 	const FVector ParentAxisX = (ParentLink.Location - CurrentLink.Location).GetSafeNormal();
 
-	const float ParentCos = (ParentAxisX | ChildAxisX);
-	const float ParentSin = (ParentAxisX | ChildAxisY);
+	const double ParentCos = (ParentAxisX | ChildAxisX);
+	const double ParentSin = (ParentAxisX | ChildAxisY);
 
-	const bool bNeedsReorient = (ParentSin < 0.f) || (ParentCos > FMath::Cos(IKChain.MinRotationAngleRadians));
+	const bool bNeedsReorient = (ParentSin < 0.0) || (ParentCos > FMath::Cos(IKChain.MinRotationAngleRadians));
 
 	// Parent Link needs to be reoriented.
 	if (bNeedsReorient)
@@ -493,8 +493,8 @@ void FIKChain::FABRIK_ApplyLinkConstraints_Backward(FIKChain& IKChain, int32 Lin
 	const FVector ParentAxisY = CurrentLink.LinkAxisZ ^ ParentAxisX;
 	const FVector ChildAxisX = (ChildLink.Location - CurrentLink.Location).GetSafeNormal();
 
-	const float ChildCos = (ChildAxisX | ParentAxisX);
-	const float ChildSin = (ChildAxisX | ParentAxisY);
+	const double ChildCos = (ChildAxisX | ParentAxisX);
+	const double ChildSin = (ChildAxisX | ParentAxisY);
 
 	const bool bNeedsReorient = (ChildSin > 0.f) || (ChildCos > FMath::Cos(IKChain.MinRotationAngleRadians));
 
@@ -524,18 +524,18 @@ void FIKChain::FABRIK_ForwardReach(const FVector& InTargetLocation, FIKChain& IK
 		FVector EndEffectorToTarget = InTargetLocation - IKChain.Links[0].Location;
 
 		FVector EndEffectorToTargetDir;
-		float EndEffectToTargetSize;
+		double EndEffectToTargetSize;
 		EndEffectorToTarget.ToDirectionAndLength(EndEffectorToTargetDir, EndEffectToTargetSize);
 
-		const float ReachStepAlpha = FMath::Clamp(CVarAnimLegIKTargetReachStepPercent.GetValueOnAnyThread(), 0.01f, 0.99f);
+		const double ReachStepAlpha = FMath::Clamp(CVarAnimLegIKTargetReachStepPercent.GetValueOnAnyThread(), 0.01, 0.99);
 
-		float Displacement = EndEffectToTargetSize;
+		double Displacement = EndEffectToTargetSize;
 		for (int32 LinkIndex = 1; LinkIndex < IKChain.NumLinks; LinkIndex++)
 		{
 			FVector EndEffectorToParent = IKChain.Links[LinkIndex].Location - IKChain.Links[0].Location;
-			float ParentDisplacement = (EndEffectorToParent | EndEffectorToTargetDir);
+			double ParentDisplacement = (EndEffectorToParent | EndEffectorToTargetDir);
 
-			Displacement = (ParentDisplacement > 0.f) ? FMath::Min(Displacement, ParentDisplacement * ReachStepAlpha) : Displacement;
+			Displacement = (ParentDisplacement > 0.0) ? FMath::Min(Displacement, ParentDisplacement * ReachStepAlpha) : Displacement;
 		}
 
 		IKChain.Links[0].Location += EndEffectorToTargetDir * Displacement;
@@ -568,15 +568,15 @@ void FIKChain::FABRIK_BackwardReach(const FVector& InRootTargetLocation, FIKChai
 		float RootToRootTargetSize;
 		RootToRootTarget.ToDirectionAndLength(RootToRootTargetDir, RootToRootTargetSize);
 
-		const float ReachStepAlpha = FMath::Clamp(CVarAnimLegIKTargetReachStepPercent.GetValueOnAnyThread(), 0.01f, 0.99f);
+		const double ReachStepAlpha = FMath::Clamp(CVarAnimLegIKTargetReachStepPercent.GetValueOnAnyThread(), 0.01, 0.99);
 
-		float Displacement = RootToRootTargetSize;
+		double Displacement = RootToRootTargetSize;
 		for (int32 LinkIndex = IKChain.NumLinks - 2; LinkIndex >= 0; LinkIndex--)
 		{
 			FVector RootToChild = IKChain.Links[IKChain.NumLinks - 2].Location - IKChain.Links.Last().Location;
-			float ChildDisplacement = (RootToChild | RootToRootTargetDir);
+			double ChildDisplacement = (RootToChild | RootToRootTargetDir);
 
-			Displacement = (ChildDisplacement > 0.f) ? FMath::Min(Displacement, ChildDisplacement * ReachStepAlpha) : Displacement;
+			Displacement = (ChildDisplacement > 0.0) ? FMath::Min(Displacement, ChildDisplacement * ReachStepAlpha) : Displacement;
 		}
 
 		IKChain.Links.Last().Location += RootToRootTargetDir * Displacement;
@@ -619,16 +619,16 @@ static FVector FindPlaneNormal(const TArray<FIKChainLink>& Links, const FVector&
 
 TAutoConsoleVariable<int32> CVarAnimLegIKAveragePull(TEXT("a.AnimNode.LegIK.AveragePull"), 1, TEXT("Leg IK AveragePull"));
 
-void FIKChain::SolveFABRIK(const FVector& InTargetLocation, float InReachPrecision, int32 InMaxIterations)
+void FIKChain::SolveFABRIK(const FVector& InTargetLocation, double InReachPrecision, int32 InMaxIterations)
 {
 	// Make sure precision is not too small.
-	const float ReachPrecision = FMath::Max(InReachPrecision, KINDA_SMALL_NUMBER);
+	const double ReachPrecision = FMath::Max(InReachPrecision, DOUBLE_KINDA_SMALL_NUMBER);
 
 	const FVector RootTargetLocation = Links.Last().Location;
-	const float PullDistributionAlpha = FMath::Clamp(CVarAnimLegIKPullDistribution.GetValueOnAnyThread(), 0.f, 1.f);
+	const double PullDistributionAlpha = FMath::Clamp(CVarAnimLegIKPullDistribution.GetValueOnAnyThread(), 0.0, 1.0);
 
 	// Check distance between foot and foot target location
-	float Slop = FVector::Dist(Links[0].Location, InTargetLocation);
+	double Slop = FVector::Dist(Links[0].Location, InTargetLocation);
 	if (Slop > ReachPrecision || bEnableRotationLimit)
 	{
 		if (bEnableRotationLimit)
@@ -671,7 +671,7 @@ void FIKChain::SolveFABRIK(const FVector& InTargetLocation, float InReachPrecisi
 		const int32 MaxIterations = FMath::Max(InMaxIterations, 1);
 		do
 		{
-			const float PreviousSlop = Slop;
+			const double PreviousSlop = Slop;
 
 #if ENABLE_ANIM_DEBUG
 			bool bDrawDebug = bShowDebug && (IterationCount == (MaxIterations - 1));
diff --git a/Engine/Source/Runtime/AnimGraphRuntime/Private/BoneControllers/AnimNode_RigidBody.cpp b/Engine/Source/Runtime/AnimGraphRuntime/Private/BoneControllers/AnimNode_RigidBody.cpp
--- a/Engine/Source/Runtime/AnimGraphRuntime/Private/BoneControllers/AnimNode_RigidBody.cpp
+++ b/Engine/Source/Runtime/AnimGraphRuntime/Private/BoneControllers/AnimNode_RigidBody.cpp
@@ -167,8 +167,8 @@ FAnimNode_RigidBody::FAnimNode_RigidBody()
 	, bComponentSpaceSimulation_DEPRECATED(true)
 #endif
 	, SimulationTiming(ESimulationTiming::Default)
-	, WorldTimeSeconds(0.0f)
-	, LastEvalTimeSeconds(0.0f)
+	, WorldTimeSeconds(0.0)
+	, LastEvalTimeSeconds(0.0)
 	, AccumulatedDeltaTime(0.0f)
 	, AnimPhysicsMinDeltaTime(0.0f)
 	, bSimulateAnimPhysicsAfterReset(false)
@@ -187,7 +187,7 @@ FAnimNode_RigidBody::FAnimNode_RigidBody()
 	, ComponentsInSim()
 	, ComponentsInSimTick(0)
 	, WorldSpaceGravity(0.0f)
-	, TotalMass(0.0f)
+	, TotalMass(0.0)
 	, CachedBounds(FVector::ZeroVector, 0.0f)
 	, QueryParams(NAME_None, FCollisionQueryParams::GetUnknownStatId())
 	, PhysScene(nullptr)
@@ -376,7 +376,7 @@ void FAnimNode_RigidBody::InitializeNewBodyTransformsDuringSimulation(FComponent
 
 				const FTransform WSBodyTM = BodyRelativeTransform * Bodies[OutputData.ParentBodyIndex]->GetWorldTransform();
 				Bodies[BodyIndex]->InitWorldTransform(WSBodyTM);
-				BodyAnimData[BodyIndex].RefPoseLength = BodyRelativeTransform.GetLocation().Size();
+				BodyAnimData[BodyIndex].RefPoseLength = static_cast<float>(BodyRelativeTransform.GetLocation().Size());
 			}
 			// If we don't have a parent body, then we can just grab the incoming pose in component space.
 			else
@@ -474,7 +474,7 @@ void FAnimNode_RigidBody::CalculateSimulationSpace(
 
 		// If we limit the angular velocity, we also need to limit the component of linear velocity that comes from (angvel x offset)
 		float AngVelScale = 1.0f;
-		float NetAngVelLenSq = NetAngVel.SizeSquared();
+		float NetAngVelLenSq = static_cast<float>(NetAngVel.SizeSquared());
 		if (NetAngVelLenSq > FMath::Square(Settings.MaxAngularVelocity))
 		{
 			AngVelScale = Settings.MaxAngularVelocity * FMath::InvSqrt(NetAngVelLenSq);
@@ -637,7 +637,7 @@ void FAnimNode_RigidBody::EvaluateSkeletalControl_AnyThread(FComponentSpacePoseC
 								
 								// Angular Velocity
 								const FQuat DeltaRotation = (NextSSTM.GetRotation().Inverse() * PrevSSTM.GetRotation());
-								const float RotationAngle = DeltaRotation.GetAngle() / DeltaSeconds;
+								const double RotationAngle = DeltaRotation.GetAngle() / DeltaSeconds;
 								BodyData.TransferedBoneAngularVelocity = (FQuat(DeltaRotation.GetRotationAxis(), RotationAngle)); 
 							}
 							else
@@ -679,7 +679,7 @@ void FAnimNode_RigidBody::EvaluateSkeletalControl_AnyThread(FComponentSpacePoseC
 						Bodies[BodyIndex]->SetWorldTransform(BodyTM);
 						if (OutputData.ParentBodyIndex != INDEX_NONE)
 						{
-							BodyAnimData[BodyIndex].RefPoseLength = BodyTM.GetRelativeTransform(Bodies[OutputData.ParentBodyIndex]->GetWorldTransform()).GetLocation().Size();
+							BodyAnimData[BodyIndex].RefPoseLength = static_cast<float>(BodyTM.GetRelativeTransform(Bodies[OutputData.ParentBodyIndex]->GetWorldTransform()).GetLocation().Size());
 						}
 					}
 				}
@@ -702,7 +702,7 @@ void FAnimNode_RigidBody::EvaluateSkeletalControl_AnyThread(FComponentSpacePoseC
 						Bodies[BodyIndex]->InitWorldTransform(BodyTM);
 						if (OutputData.ParentBodyIndex != INDEX_NONE)
 						{
-							BodyAnimData[BodyIndex].RefPoseLength = BodyTM.GetRelativeTransform(Bodies[OutputData.ParentBodyIndex]->GetWorldTransform()).GetLocation().Size();
+							BodyAnimData[BodyIndex].RefPoseLength = static_cast<float>(BodyTM.GetRelativeTransform(Bodies[OutputData.ParentBodyIndex]->GetWorldTransform()).GetLocation().Size());
 						}
 					}
 				}
@@ -773,8 +773,8 @@ void FAnimNode_RigidBody::EvaluateSkeletalControl_AnyThread(FComponentSpacePoseC
 							ImmediatePhysics::FActorHandle* Body = Bodies[BodyIndex];
 
 							// Apply 
-							const float BodyInvMass = Body->GetInverseMass();
-							if (BodyInvMass > 0.f)
+							const double BodyInvMass = Body->GetInverseMass();
+							if (BodyInvMass > 0.0)
 							{
 								// Final desired acceleration to apply to body
 								FVector FinalBodyLinearAcc = ApplyLinearAcc;
@@ -897,7 +897,7 @@ void FAnimNode_RigidBody::EvaluateSkeletalControl_AnyThread(FComponentSpacePoseC
 
 					// get local transform
 					FTransform LocalTransform = BodyTM.GetRelativeTransform(ParentTransform);
-					const float CurrentLength = LocalTransform.GetTranslation().Size();
+					const float CurrentLength = static_cast<float>(LocalTransform.GetTranslation().Size());
 
 					// this is inconsistent with constraint. The actual linear limit is set by constraint
 					if (!FMath::IsNearlyEqual(CurrentLength, BodyAnimData[BodyIndex].RefPoseLength, KINDA_SMALL_NUMBER))
@@ -1055,7 +1055,7 @@ void FAnimNode_RigidBody::InitPhysics(const UAnimInstance* InAnimInstance)
 		Bodies.Empty(NumBodies);
 		BodyAnimData.Reset(NumBodies);
 		BodyAnimData.AddDefaulted(NumBodies);
-		TotalMass = 0.f;
+		TotalMass = 0.0;
 
 		// Instantiate a FBodyInstance/FConstraintInstance set that will be cloned into the Immediate Physics sim.
 		// NOTE: We do not have a skeleton at the moment, so we have to use the ref pose
@@ -1118,8 +1118,8 @@ void FAnimNode_RigidBody::InitPhysics(const UAnimInstance* InAnimInstance)
 				{
 					if (bSimulated)
 					{
-						const float InvMass = NewBodyHandle->GetInverseMass();
-						TotalMass += InvMass > 0.f ? 1.f / InvMass : 0.f;
+						const double InvMass = NewBodyHandle->GetInverseMass();
+						TotalMass += InvMass > 0.0 ? 1.0 / InvMass : 0.0;
 					}
 					const int32 BodyIndex = Bodies.Add(NewBodyHandle);
 					const int32 SkeletonBoneIndex = MeshToSkeletonBoneIndex[InsertBone];
@@ -1207,7 +1207,7 @@ void FAnimNode_RigidBody::InitPhysics(const UAnimInstance* InAnimInstance)
 							//set limit to ref pose 
 							FTransform Body1Transform = Body1Handle->GetWorldTransform();
 							FTransform Body2Transform = Body2Handle->GetWorldTransform();
-							BodyAnimData[BodyIndex].RefPoseLength = Body1Transform.GetRelativeTransform(Body2Transform).GetLocation().Size();
+							BodyAnimData[BodyIndex].RefPoseLength = static_cast<float>(Body1Transform.GetRelativeTransform(Body2Transform).GetLocation().Size());
 						}
 
 						if (CI->IsCollisionDisabled())
@@ -1331,17 +1331,17 @@ void FAnimNode_RigidBody::UpdateWorldForces(const FTransform& ComponentToWorld,
 	LLM_SCOPE_BYNAME(TEXT("Animation/RigidBody")); 
 	SCOPE_CYCLE_COUNTER(STAT_ImmediateUpdateWorldForces);
 
-	if(TotalMass > 0.f)
+	if(TotalMass > 0.0)
 	{
 		for (const USkeletalMeshComponent::FPendingRadialForces& PendingRadialForce : PendingRadialForces)
 		{
 			const FVector RadialForceOrigin = WorldPositionToSpace(SimulationSpace, PendingRadialForce.Origin, ComponentToWorld, BaseBoneTM);
 			for(ImmediatePhysics::FActorHandle* Body : Bodies)
 			{
-				const float InvMass = Body->GetInverseMass();
-				if(InvMass > 0.f)
+				const double InvMass = Body->GetInverseMass();
+				if(InvMass > 0.0)
 				{
-					const float StrengthPerBody = PendingRadialForce.bIgnoreMass ? PendingRadialForce.Strength : PendingRadialForce.Strength / (TotalMass * InvMass);
+					const double StrengthPerBody = PendingRadialForce.bIgnoreMass ? PendingRadialForce.Strength : PendingRadialForce.Strength / (TotalMass * InvMass);
 					ImmediatePhysics::EForceType ForceType;
 					if (PendingRadialForce.Type == USkeletalMeshComponent::FPendingRadialForces::AddImpulse)
 					{
@@ -1362,8 +1362,8 @@ void FAnimNode_RigidBody::UpdateWorldForces(const FTransform& ComponentToWorld,
 			const FVector ExternalForceInSimSpace = WorldVectorToSpaceNoScale(SimulationSpace, ExternalForce, ComponentToWorld, BaseBoneTM);
 			for (ImmediatePhysics::FActorHandle* Body : Bodies)
 			{
-				const float InvMass = Body->GetInverseMass();
-				if (InvMass > 0.f)
+				const double InvMass = Body->GetInverseMass();
+				if (InvMass > 0.0)
 				{
 					Body->AddForce(ExternalForceInSimSpace);
 				}
@@ -1502,8 +1502,10 @@ void FAnimNode_RigidBody::PreUpdate(const UAnimInstance* InAnimInstance)
 			{
 				const FBox BoundingBox = SKC->CalcBounds(SKC->GetComponentTransform()).GetBox();
 
+				PRAGMA_DISABLE_UNSAFE_TYPECAST_WARNINGS
 				World->PhysicsField->FillTransientCommands(false, BoundingBox, WorldTimeSeconds, PerSolverField.GetTransientCommands());
 				World->PhysicsField->FillPersistentCommands(false, BoundingBox, WorldTimeSeconds, PerSolverField.GetPersistentCommands());
+				PRAGMA_RESTORE_UNSAFE_TYPECAST_WARNINGS
 			}
 		}
 	}
@@ -1635,7 +1637,7 @@ void FAnimNode_RigidBody::CollectWorldObjects()
 	{
 		// @todo(ccaulfield): should this use CachedBounds?
 		TArray<FOverlapResult> Overlaps;
-		UnsafeWorld->OverlapMultiByChannel(Overlaps, CachedBounds.Center, FQuat::Identity, OverlapChannel, FCollisionShape::MakeSphere(CachedBounds.W), QueryParams, FCollisionResponseParams(ECR_Overlap));
+		UnsafeWorld->OverlapMultiByChannel(Overlaps, CachedBounds.Center, FQuat::Identity, OverlapChannel, FCollisionShape::MakeSphere(static_cast<float>(CachedBounds.W)), QueryParams, FCollisionResponseParams(ECR_Overlap));
 
 		for (const FOverlapResult& Overlap : Overlaps)
 		{
@@ -1779,8 +1781,8 @@ void FAnimNode_RigidBody::InitializeBoneReferences(const FBoneContainer& Require
 	for (int32 Index = 0; Index < NumRequiredBoneIndices; ++Index)
 	{
 		const FCompactPoseBoneIndex CompactPoseBoneIndex(Index);
-		const FBoneIndexType SkeletonBoneIndex = RequiredBones.GetSkeletonIndex(CompactPoseBoneIndex);
-		const FBoneIndexType IndexToBodyNum = SkeletonBoneIndexToBodyIndex.Num();
+		const int32 SkeletonBoneIndex = RequiredBones.GetSkeletonIndex(CompactPoseBoneIndex);
+		const int32 IndexToBodyNum = SkeletonBoneIndexToBodyIndex.Num();
 
 		// If we have a missing bone in our skeleton, we don't want to have an out of bounds access.
 		if (SkeletonBoneIndex >= IndexToBodyNum)
@@ -1811,7 +1813,7 @@ void FAnimNode_RigidBody::InitializeBoneReferences(const FBoneContainer& Require
 			FCompactPoseBoneIndex CompactParentIndex = RequiredBones.GetParentBoneIndex(CompactPoseBoneIndex);
 			while (CompactParentIndex != INDEX_NONE)
 			{
-				const FBoneIndexType SkeletonParentBoneIndex = RequiredBones.GetSkeletonIndex(CompactParentIndex);
+				const int32 SkeletonParentBoneIndex = RequiredBones.GetSkeletonIndex(CompactParentIndex);
 
 				// Must check our parent as well for a missing bone.
 				if (SkeletonParentBoneIndex >= IndexToBodyNum)
diff --git a/Engine/Source/Runtime/AnimGraphRuntime/Private/BoneControllers/AnimNode_SplineIK.cpp b/Engine/Source/Runtime/AnimGraphRuntime/Private/BoneControllers/AnimNode_SplineIK.cpp
--- a/Engine/Source/Runtime/AnimGraphRuntime/Private/BoneControllers/AnimNode_SplineIK.cpp
+++ b/Engine/Source/Runtime/AnimGraphRuntime/Private/BoneControllers/AnimNode_SplineIK.cpp
@@ -217,22 +217,22 @@ void FAnimNode_SplineIK::TransformSpline()
 
 float FAnimNode_SplineIK::FindParamAtFirstSphereIntersection(const FVector& InOrigin, float InRadius, int32& StartingLinearIndex)
 {
-	const float RadiusSquared = InRadius * InRadius;
+	const double RadiusSquared = InRadius * InRadius;
 	const int32 LinearCount = LinearApproximation.Num() - 1;
 	for (int32 LinearIndex = StartingLinearIndex; LinearIndex < LinearCount; ++LinearIndex)
 	{
 		const FSplinePositionLinearApproximation& LinearPoint = LinearApproximation[LinearIndex];
 		const FSplinePositionLinearApproximation& NextLinearPoint = LinearApproximation[LinearIndex + 1];
 
-		const float InnerDistanceSquared = (InOrigin - LinearPoint.Position).SizeSquared();
-		const float OuterDistanceSquared = (InOrigin - NextLinearPoint.Position).SizeSquared();
+		const double InnerDistanceSquared = (InOrigin - LinearPoint.Position).SizeSquared();
+		const double OuterDistanceSquared = (InOrigin - NextLinearPoint.Position).SizeSquared();
 		if (InnerDistanceSquared <= RadiusSquared && OuterDistanceSquared >= RadiusSquared)
 		{
 			StartingLinearIndex = LinearIndex;
 
-			const float InnerDistance = FMath::Sqrt(InnerDistanceSquared);
-			const float OuterDistance = FMath::Sqrt(OuterDistanceSquared);
-			const float InterpParam = FMath::Clamp((InRadius - InnerDistance) / (OuterDistance - InnerDistance), 0.0f, 1.0f);
+			const double InnerDistance = FMath::Sqrt(InnerDistanceSquared);
+			const double OuterDistance = FMath::Sqrt(OuterDistanceSquared);
+			const double InterpParam = FMath::Clamp((InRadius - InnerDistance) / (OuterDistance - InnerDistance), 0.0f, 1.0f);
 			
 			return FMath::Lerp(LinearPoint.SplineParam, NextLinearPoint.SplineParam, InterpParam);
 		}
@@ -302,7 +302,7 @@ void FAnimNode_SplineIK::BuildBoneSpline(const FReferenceSkeleton& RefSkeleton)
 		CachedOffsetRotations.Reset();
 		for (int32 BoneIndex = 0; BoneIndex < CachedBoneReferences.Num(); BoneIndex++)
 		{
-			float BoneLength = 0.0f;
+			double BoneLength = 0.0f;
 			FQuat BoneOffsetRotation = FQuat::Identity;
 
 			if (BoneIndex > 0)
@@ -319,8 +319,8 @@ void FAnimNode_SplineIK::BuildBoneSpline(const FReferenceSkeleton& RefSkeleton)
 				FVector TransformedAxis = Transform.GetRotation().RotateVector(FMatrix::Identity.GetUnitAxis((EAxis::Type)BoneAxis)).GetSafeNormal();
 				BoneOffsetRotation = FQuat::FindBetweenNormals(BoneDir.GetSafeNormal(), TransformedAxis);
 			}
-
-			CachedBoneLengths.Add(BoneLength);
+			
+			CachedBoneLengths.Add(static_cast<float>(BoneLength));
 			CachedOffsetRotations.Add(BoneOffsetRotation);
 		}
 
diff --git a/Engine/Source/Runtime/AnimGraphRuntime/Private/BoneControllers/AnimNode_SpringBone.cpp b/Engine/Source/Runtime/AnimGraphRuntime/Private/BoneControllers/AnimNode_SpringBone.cpp
--- a/Engine/Source/Runtime/AnimGraphRuntime/Private/BoneControllers/AnimNode_SpringBone.cpp
+++ b/Engine/Source/Runtime/AnimGraphRuntime/Private/BoneControllers/AnimNode_SpringBone.cpp
@@ -11,10 +11,10 @@
 // FAnimNode_SpringBone
 
 FAnimNode_SpringBone::FAnimNode_SpringBone()
-	: MaxDisplacement(0.0f)
-	, SpringStiffness(50.0f)
-	, SpringDamping(4.0f)
-	, ErrorResetThresh(256.0f)
+	: MaxDisplacement(0.0)
+	, SpringStiffness(50.0)
+	, SpringDamping(4.0)
+	, ErrorResetThresh(256.0)
 	, BoneLocation(FVector::ZeroVector)
 	, BoneVelocity(FVector::ZeroVector)
 	, OwnerVelocity(FVector::ZeroVector)
@@ -139,10 +139,10 @@ void FAnimNode_SpringBone::EvaluateSkeletalControl_AnyThread(FComponentSpacePose
 
 			// Integrate velocity
 			// Make sure damping with variable frame rate actually dampens velocity. Otherwise Spring will go nuts.
-			float const CutOffDampingValue = 1.f / FixedTimeStep;
+			double const CutOffDampingValue = 1.0 / FixedTimeStep;
 			if (SpringDamping > CutOffDampingValue)
 			{
-				float const SafetyScale = CutOffDampingValue / SpringDamping;
+				double const SafetyScale = CutOffDampingValue / SpringDamping;
 				BoneVelocity += SafetyScale * (Acceleration * FixedTimeStep);
 			}
 			else
@@ -151,7 +151,7 @@ void FAnimNode_SpringBone::EvaluateSkeletalControl_AnyThread(FComponentSpacePose
 			}
 
 			// Clamp velocity to something sane (|dX/dt| <= ErrorResetThresh)
-			float const BoneVelocityMagnitude = BoneVelocity.Size();
+			double const BoneVelocityMagnitude = BoneVelocity.Size();
 			if (BoneVelocityMagnitude * FixedTimeStep > ErrorResetThresh)
 			{
 				BoneVelocity *= (ErrorResetThresh / (BoneVelocityMagnitude * FixedTimeStep));
diff --git a/Engine/Source/Runtime/AnimGraphRuntime/Private/BoneControllers/AnimNode_Trail.cpp b/Engine/Source/Runtime/AnimGraphRuntime/Private/BoneControllers/AnimNode_Trail.cpp
--- a/Engine/Source/Runtime/AnimGraphRuntime/Private/BoneControllers/AnimNode_Trail.cpp
+++ b/Engine/Source/Runtime/AnimGraphRuntime/Private/BoneControllers/AnimNode_Trail.cpp
@@ -115,7 +115,7 @@ void FAnimNode_Trail::EvaluateSkeletalControl_AnyThread(FComponentSpacePoseConte
 	{
 		for(int32 i=0; i<ChainBoneIndices.Num(); i++)
 		{
-			if (BoneContainer.Contains(ChainBoneIndices[i]))
+			if (BoneContainer.Contains(IntCastChecked<FBoneIndexType>(ChainBoneIndices[i])))
 			{
 				FCompactPoseBoneIndex ChildIndex = BoneContainer.MakeCompactPoseIndex(FMeshPoseBoneIndex(ChainBoneIndices[i]));
 				const FTransform& ChainTransform = Output.Pose.GetComponentSpaceTransform(ChildIndex);
@@ -202,12 +202,12 @@ void FAnimNode_Trail::EvaluateSkeletalControl_AnyThread(FComponentSpacePoseConte
 		// If desired, prevent bones stretching too far.
 		if (bLimitStretch)
 		{
-			float RefPoseLength = TargetDelta.Size();
+			double RefPoseLength = TargetDelta.Size();
 			FVector CurrentDelta = TrailBoneLocations[i] - TrailBoneLocations[i - 1];
-			float CurrentLength = CurrentDelta.Size();
+			double CurrentLength = CurrentDelta.Size();
 
 			// If we are too far - cut it back (just project towards parent particle).
-			if ((CurrentLength - RefPoseLength > StretchLimit) && CurrentLength > SMALL_NUMBER)
+			if ((CurrentLength - RefPoseLength > StretchLimit) && CurrentLength > UE_DOUBLE_SMALL_NUMBER)
 			{
 				FVector CurrentDir = CurrentDelta / CurrentLength;
 				TrailBoneLocations[i] = TrailBoneLocations[i - 1] + (CurrentDir * (RefPoseLength + StretchLimit));
@@ -234,8 +234,8 @@ void FAnimNode_Trail::EvaluateSkeletalControl_AnyThread(FComponentSpacePoseConte
 #if WITH_EDITORONLY_DATA				
 				DebugPlaneTransforms[Index] = LimitPlaneTransform;
 #endif // #if WITH_EDITORONLY_DATA
-				float DistanceFromPlane = LimitPlane.PlaneDot(TrailBoneLocations[i]);
-				if (DistanceFromPlane < 0)
+				const double DistanceFromPlane = LimitPlane.PlaneDot(TrailBoneLocations[i]);
+				if (DistanceFromPlane < 0.0)
 				{
 					TrailBoneLocations[i] -= DistanceFromPlane*FVector(LimitPlane.X, LimitPlane.Y, LimitPlane.Z);
 				}
@@ -250,11 +250,11 @@ void FAnimNode_Trail::EvaluateSkeletalControl_AnyThread(FComponentSpacePoseConte
 		if (bReorientParentToChild)
 		{
 			FVector CurrentBoneDir = OutBoneTransforms[i - 1].Transform.TransformVector(GetAlignVector(ChainBoneAxis, bInvertChainBoneAxis));
-			CurrentBoneDir = CurrentBoneDir.GetSafeNormal(SMALL_NUMBER);
+			CurrentBoneDir = CurrentBoneDir.GetSafeNormal(UE_DOUBLE_SMALL_NUMBER);
 
 			// Calculate vector from parent to child.
 			FVector DeltaTranslation = OutBoneTransforms[i].Transform.GetTranslation() - OutBoneTransforms[i - 1].Transform.GetTranslation();
-			FVector NewBoneDir = FVector(DeltaTranslation).GetSafeNormal(SMALL_NUMBER);
+			FVector NewBoneDir = FVector(DeltaTranslation).GetSafeNormal(UE_DOUBLE_SMALL_NUMBER);
 
 			// Calculate a quaternion that gets us from our current rotation to the desired one.
 			FQuat DeltaLookQuat = FQuat::FindBetweenNormals(CurrentBoneDir, NewBoneDir);
@@ -389,7 +389,7 @@ bool FAnimNode_Trail::IsValidToEvaluate(const USkeleton* Skeleton, const FBoneCo
 				// originally tried in AnimGraphNode, but that doesn't know hierarchy so I can't verify it there. Maybe should try with USkeleton asset there. @todo
 				return false;
 			}
-			else if (RequiredBones.Contains(ChainIndex) == false)
+			else if (RequiredBones.Contains(IntCastChecked<FBoneIndexType>(ChainIndex)) == false)
 			{
 				return false;
 			}
diff --git a/Engine/Source/Runtime/AnimGraphRuntime/Private/BoneControllers/AnimNode_TwistCorrectiveNode.cpp b/Engine/Source/Runtime/AnimGraphRuntime/Private/BoneControllers/AnimNode_TwistCorrectiveNode.cpp
--- a/Engine/Source/Runtime/AnimGraphRuntime/Private/BoneControllers/AnimNode_TwistCorrectiveNode.cpp
+++ b/Engine/Source/Runtime/AnimGraphRuntime/Private/BoneControllers/AnimNode_TwistCorrectiveNode.cpp
@@ -107,24 +107,24 @@ float FAnimNode_TwistCorrectiveNode::GetAngle(const FVector& Base, const FVector
 	FVector TwistPlaneNormal = TwistPlaneNormalAxis.GetTransformedAxis(ReferencetBoneTransform);
 
 	// find out if it's same direction. If not, we're clamped
-	float BaseAngle = 0.f, TwistAngle = 0.f;
+	float BaseAngle = 0.0, TwistAngle = 0.0;
 
 	// if facing same direction we care
-	float BaseDotProduct = FVector::DotProduct(TwistPlaneNormal, Base);
-	if (BaseDotProduct > 0.f)
+	float BaseDotProduct = static_cast<float>(FVector::DotProduct(TwistPlaneNormal, Base));
+	if (BaseDotProduct > 0.0)
 	{
 		// http://www.vitutor.com/geometry/distance/line_plane.html
 		BaseAngle = FMath::Asin(BaseDotProduct);
 	}
 
-	float TwistDotProduct = FVector::DotProduct(TwistPlaneNormal, Twist);
-	if (TwistDotProduct > 0.f)
+	float TwistDotProduct = static_cast<float>(FVector::DotProduct(TwistPlaneNormal, Twist));
+	if (TwistDotProduct > 0.0)
 	{
 		// http://www.vitutor.com/geometry/distance/line_plane.html
 		TwistAngle = FMath::Asin(TwistDotProduct);
 	}
 
-	return (TwistAngle - BaseAngle);
+	return TwistAngle - BaseAngle;
 }
 
 /////////////////////////////////////////////////////////
diff --git a/Engine/Source/Runtime/AnimGraphRuntime/Private/CommonAnimationLibrary.cpp b/Engine/Source/Runtime/AnimGraphRuntime/Private/CommonAnimationLibrary.cpp
--- a/Engine/Source/Runtime/AnimGraphRuntime/Private/CommonAnimationLibrary.cpp
+++ b/Engine/Source/Runtime/AnimGraphRuntime/Private/CommonAnimationLibrary.cpp
@@ -9,7 +9,7 @@
 
 float CommonAnimationLibrary::ScalarEasing(float Value, const FRuntimeFloatCurve& CustomCurve, EEasingFuncType EasingType, bool bFlip, float Weight)
 {
-	float Original = Value;
+	const float Original = Value;
 	float Result = Value;
 
 	if(bFlip)
@@ -26,7 +26,7 @@ float CommonAnimationLibrary::ScalarEasing(float Value, const FRuntimeFloatCurve
 		}
 		case EEasingFuncType::Sinusoidal:
 		{
-			Result = FMath::Clamp<float>((FMath::Sin(Value * PI - HALF_PI) + 1.f) / 2.f, 0.f, 1.f);
+			Result = FMath::Clamp<float>((FMath::Sin(Value * UE_PI - UE_HALF_PI) + 1.f) / 2.f, 0.f, 1.f);
 			break;
 		}
 		case EEasingFuncType::Cubic:
@@ -122,14 +122,14 @@ FVector CommonAnimationLibrary::RetargetSingleLocation(
 	float TargetMaximum
 )
 {
-	float SourceDelta = SourceMaximum - SourceMinimum;
-	float TargetDelta = TargetMaximum - TargetMinimum;
+	const float SourceDelta = SourceMaximum - SourceMinimum;
+	const float TargetDelta = TargetMaximum - TargetMinimum;
 	if (FMath::IsNearlyEqual(SourceDelta, 0.f) || FMath::IsNearlyEqual(TargetDelta, 0.f))
 	{
 		return Location;
 	}
 
-	float AxisLengthSquared = Axis.SizeSquared();
+	const float AxisLengthSquared =static_cast<float>(Axis.SizeSquared());
 	if (FMath::IsNearlyEqual(AxisLengthSquared, 0.f))
 	{
 		Axis = FVector(1.f, 0.f, 0.f);
@@ -142,13 +142,13 @@ FVector CommonAnimationLibrary::RetargetSingleLocation(
 
 	Location = Source.InverseTransformPosition(Location);
 
-	float Delta = FVector::DotProduct(Location, Axis);
+	const float Delta = static_cast<float>(FVector::DotProduct(Location, Axis));
 	float OutBias = (Delta - SourceMinimum) / SourceDelta;
 	if (EasingType != EEasingFuncType::Linear)
 	{
 		OutBias = ScalarEasing(OutBias, CustomCurve, EasingType, bFlipEasing, EasingWeight);
 	}
-	float TargetBias = FMath::Lerp<float>(TargetMinimum, TargetMaximum, OutBias);
+	const float TargetBias = FMath::Lerp<float>(TargetMinimum, TargetMaximum, OutBias);
 
 	Location = Location + Axis * (TargetBias - Delta);
 
@@ -174,8 +174,8 @@ FQuat CommonAnimationLibrary::RetargetSingleRotation(
 {
 	FQuat Rotation = RotationIn;
 
-	float SourceDelta = SourceMaximum - SourceMinimum;
-	float TargetDelta = TargetMaximum - TargetMinimum;
+	const float SourceDelta = SourceMaximum - SourceMinimum;
+	const float TargetDelta = TargetMaximum - TargetMinimum;
 	if (FMath::IsNearlyEqual(SourceDelta, 0.f) || FMath::IsNearlyEqual(TargetDelta, 0.f))
 	{
 		return Rotation;
@@ -192,30 +192,30 @@ FQuat CommonAnimationLibrary::RetargetSingleRotation(
 		case ERotationComponent::EulerX:
 		{
 			Euler = Rotation.Euler();
-			Angle = Euler.X;
+			Angle = static_cast<float>(Euler.X);
 			break;
 		}
 		case ERotationComponent::EulerY:
 		{
 			Euler = Rotation.Euler();
-			Angle = Euler.Y;
+			Angle = static_cast<float>(Euler.Y);
 			break;
 		}
 		case ERotationComponent::EulerZ:
 		{
 			Euler = Rotation.Euler();
-			Angle = Euler.Z;
+			Angle = static_cast<float>(Euler.Z);
 			break;
 		}
 		case ERotationComponent::QuaternionAngle:
 		{
-			Angle = FMath::RadiansToDegrees(Rotation.GetAngle());
+			Angle = static_cast<float>(FMath::RadiansToDegrees(Rotation.GetAngle()));
 			break;
 		}
 		case ERotationComponent::SwingAngle:
 		case ERotationComponent::TwistAngle:
 		{
-			float AxisLengthSquared = TwistAxis.SizeSquared();
+			const float AxisLengthSquared = static_cast<float>(TwistAxis.SizeSquared());
 			if (FMath::IsNearlyEqual(AxisLengthSquared, 0.f))
 			{
 				TwistAxis = FVector(1.f, 0.f, 0.f);
@@ -228,24 +228,24 @@ FQuat CommonAnimationLibrary::RetargetSingleRotation(
 			Rotation.ToSwingTwist(TwistAxis, Swing, Twist);
 			if (RotationComponent == ERotationComponent::SwingAngle)
 			{
-				Angle = FMath::RadiansToDegrees(Swing.GetAngle());
+				Angle = FMath::RadiansToDegrees(static_cast<float>(Swing.GetAngle()));
 			}
 			else
 			{
-				Angle = FMath::RadiansToDegrees(Twist.GetAngle());
+				Angle = FMath::RadiansToDegrees(static_cast<float>(Twist.GetAngle()));
 			}
 			break;
 		}
 	}
 
-	float AngleSign = Angle < 0.f ? -1.f : 1.f;
+	const float AngleSign = Angle < 0.f ? -1.f : 1.f;
 	if (bUseAbsoluteAngle)
 	{
 		Angle = FMath::Abs<float>(Angle);
 	}
 
 	float OutBias = (Angle - SourceMinimum) / SourceDelta;
-	float TempBias = OutBias;
+	const float TempBias = OutBias;
 	if (EasingType != EEasingFuncType::Linear)
 	{
 		OutBias = ScalarEasing(OutBias, CustomCurve, EasingType, bFlipEasing, EasingWeight);
diff --git a/Engine/Source/Runtime/AnimGraphRuntime/Private/KismetAnimationLibrary.cpp b/Engine/Source/Runtime/AnimGraphRuntime/Private/KismetAnimationLibrary.cpp
--- a/Engine/Source/Runtime/AnimGraphRuntime/Private/KismetAnimationLibrary.cpp
+++ b/Engine/Source/Runtime/AnimGraphRuntime/Private/KismetAnimationLibrary.cpp
@@ -67,7 +67,7 @@ float UKismetAnimationLibrary::K2_DistanceBetweenTwoSocketsAndMapRange(const USk
 		FTransform SocketTransformA = Component->GetSocketTransform(SocketOrBoneNameA, SocketSpaceA);
 		FTransform SocketTransformB = Component->GetSocketTransform(SocketOrBoneNameB, SocketSpaceB);
 
-		float Distance = (SocketTransformB.GetLocation() - SocketTransformA.GetLocation()).Size();
+		const float Distance = static_cast<float>((SocketTransformB.GetLocation() - SocketTransformA.GetLocation()).Size());
 
 		if (bRemapRange)
 		{
@@ -145,7 +145,7 @@ float UKismetAnimationLibrary::K2_CalculateVelocityFromPositionHistory(
 	}
 	else
 	{
-		float LengthOfV = ((Position - History.Positions[History.LastIndex]) / DeltaSeconds).Size();
+		const float LengthOfV = static_cast<float>(((Position - History.Positions[History.LastIndex]) / DeltaSeconds).Size());
 
 		if (History.Positions.Num() == NumberOfSamples)
 		{
@@ -271,21 +271,21 @@ float UKismetAnimationLibrary::CalculateDirection(const FVector& Velocity, const
 {
 	if (!Velocity.IsNearlyZero())
 	{
-		FMatrix RotMatrix = FRotationMatrix(BaseRotation);
-		FVector ForwardVector = RotMatrix.GetScaledAxis(EAxis::X);
-		FVector RightVector = RotMatrix.GetScaledAxis(EAxis::Y);
-		FVector NormalizedVel = Velocity.GetSafeNormal2D();
+		const FMatrix RotMatrix = FRotationMatrix(BaseRotation);
+		const FVector ForwardVector = RotMatrix.GetScaledAxis(EAxis::X);
+		const FVector RightVector = RotMatrix.GetScaledAxis(EAxis::Y);
+		const FVector NormalizedVel = Velocity.GetSafeNormal2D();
 
 		// get a cos(alpha) of forward vector vs velocity
-		float ForwardCosAngle = FVector::DotProduct(ForwardVector, NormalizedVel);
+		const float ForwardCosAngle = static_cast<float>(FVector::DotProduct(ForwardVector, NormalizedVel));
 		// now get the alpha and convert to degree
 		float ForwardDeltaDegree = FMath::RadiansToDegrees(FMath::Acos(ForwardCosAngle));
 
 		// depending on where right vector is, flip it
-		float RightCosAngle = FVector::DotProduct(RightVector, NormalizedVel);
-		if (RightCosAngle < 0)
+		const float RightCosAngle = static_cast<float>(FVector::DotProduct(RightVector, NormalizedVel));
+		if (RightCosAngle < 0.f)
 		{
-			ForwardDeltaDegree *= -1;
+			ForwardDeltaDegree *= -1.f;
 		}
 
 		return ForwardDeltaDegree;
diff --git a/Engine/Source/Runtime/AnimGraphRuntime/Private/RBF/RBFSolver.cpp b/Engine/Source/Runtime/AnimGraphRuntime/Private/RBF/RBFSolver.cpp
--- a/Engine/Source/Runtime/AnimGraphRuntime/Private/RBF/RBFSolver.cpp
+++ b/Engine/Source/Runtime/AnimGraphRuntime/Private/RBF/RBFSolver.cpp
@@ -49,18 +49,18 @@ void FRBFEntry::AddFromRotator(const FRotator& InRot)
 {
 	const int32 BaseIndex = Values.AddUninitialized(3);
 
-	Values[BaseIndex + 0] = InRot.Roll;
-	Values[BaseIndex + 1] = InRot.Pitch;
-	Values[BaseIndex + 2] = InRot.Yaw;
+	Values[BaseIndex + 0] = static_cast<float>(InRot.Roll);
+	Values[BaseIndex + 1] = static_cast<float>(InRot.Pitch);
+	Values[BaseIndex + 2] = static_cast<float>(InRot.Yaw);
 }
 
 void FRBFEntry::AddFromVector(const FVector& InVector)
 {
 	const int32 BaseIndex = Values.AddUninitialized(3);
 
-	Values[BaseIndex + 0] = InVector.X;
-	Values[BaseIndex + 1] = InVector.Y;
-	Values[BaseIndex + 2] = InVector.Z;
+	Values[BaseIndex + 0] = static_cast<float>(InVector.X);
+	Values[BaseIndex + 1] = static_cast<float>(InVector.Y);
+	Values[BaseIndex + 2] = static_cast<float>(InVector.Z);
 }
 
 //////////////////////////////////////////////////////////////////////////
@@ -118,20 +118,20 @@ static float GetDistanceBetweenEntries(
 		switch (DistanceMetric)
 		{
 		case ERBFDistanceMethod::Euclidean:
-			Distance = RBFDistanceMetric::Euclidean(A.AsRotator(i), B.AsRotator(i));
+			Distance = static_cast<float>(RBFDistanceMetric::Euclidean(A.AsRotator(i), B.AsRotator(i)));
 			break;
 
 		case ERBFDistanceMethod::Quaternion:
-			Distance = RBFDistanceMetric::ArcLength(A.AsQuat(i), B.AsQuat(i));
+			Distance = static_cast<float>(RBFDistanceMetric::ArcLength(A.AsQuat(i), B.AsQuat(i)));
 			break;
 
 		case ERBFDistanceMethod::SwingAngle:
 		case ERBFDistanceMethod::DefaultMethod:
-			Distance += RBFDistanceMetric::SwingAngle(A.AsQuat(i), B.AsQuat(i), TwistAxis);
+			Distance += static_cast<float>(RBFDistanceMetric::SwingAngle(A.AsQuat(i), B.AsQuat(i), TwistAxis));
 			break;
 
 		case ERBFDistanceMethod::TwistAngle:
-			Distance += RBFDistanceMetric::TwistAngle(A.AsQuat(i), B.AsQuat(i), TwistAxis);
+			Distance += static_cast<float>(RBFDistanceMetric::TwistAngle(A.AsQuat(i), B.AsQuat(i), TwistAxis));
 			break;
 		}
 
diff --git a/Engine/Source/Runtime/AnimGraphRuntime/Public/AnimNodes/AnimNode_BlendSpacePlayer.h b/Engine/Source/Runtime/AnimGraphRuntime/Public/AnimNodes/AnimNode_BlendSpacePlayer.h
--- a/Engine/Source/Runtime/AnimGraphRuntime/Public/AnimNodes/AnimNode_BlendSpacePlayer.h
+++ b/Engine/Source/Runtime/AnimGraphRuntime/Public/AnimNodes/AnimNode_BlendSpacePlayer.h
@@ -276,7 +276,7 @@ public:
 	virtual bool ShouldResetPlayTimeWhenBlendSpaceChanges() const override { return bResetPlayTimeWhenBlendSpaceChanges; }
 	virtual bool SetResetPlayTimeWhenBlendSpaceChanges(bool bReset) override { bResetPlayTimeWhenBlendSpaceChanges = bReset; return true; }
 	virtual bool SetBlendSpace(UBlendSpace* InBlendSpace) override { BlendSpace = InBlendSpace; return true; }
-	virtual bool SetPosition(FVector InPosition) override { X = InPosition[0]; Y = InPosition[1]; return true; }
+	virtual bool SetPosition(FVector InPosition) override { X = static_cast<float>(InPosition[0]); Y = static_cast<float>(InPosition[1]); return true; }
 	virtual bool SetPlayRate(float InPlayRate) override { PlayRate = InPlayRate; return true; }
 	virtual bool SetLoop(bool bInLoop) override { bLoop = bInLoop; return true; }
 	// End of FAnimNode_BlendSpacePlayerBase interface
diff --git a/Engine/Source/Runtime/AnimGraphRuntime/Public/AnimNodes/AnimNode_PoseDriver.h b/Engine/Source/Runtime/AnimGraphRuntime/Public/AnimNodes/AnimNode_PoseDriver.h
--- a/Engine/Source/Runtime/AnimGraphRuntime/Public/AnimNodes/AnimNode_PoseDriver.h
+++ b/Engine/Source/Runtime/AnimGraphRuntime/Public/AnimNodes/AnimNode_PoseDriver.h
@@ -104,7 +104,7 @@ struct ANIMGRAPHRUNTIME_API FPoseDriverTarget
 	FName DrivenName;
 
 	/** Cached curve UID, if DriveOutput is set to DriveCurves */
-	int32 DrivenUID;
+	USkeleton::AnimCurveUID DrivenUID = SmartName::MaxUID;
 
 	/** Pose buffer index */
 	int32 PoseCurveIndex;
@@ -120,7 +120,7 @@ struct ANIMGRAPHRUNTIME_API FPoseDriverTarget
 		, FunctionType(ERBFFunctionType::DefaultFunction)
 		, bApplyCustomCurve(false)
 		, DrivenName(NAME_None)
-		, DrivenUID(INDEX_NONE)
+		, DrivenUID(SmartName::MaxUID)
 		, PoseCurveIndex(INDEX_NONE)
 		, bIsHidden(false)
 	{}
diff --git a/Engine/Source/Runtime/AnimGraphRuntime/Public/BoneControllers/AnimNode_BoneDrivenController.h b/Engine/Source/Runtime/AnimGraphRuntime/Public/BoneControllers/AnimNode_BoneDrivenController.h
--- a/Engine/Source/Runtime/AnimGraphRuntime/Public/BoneControllers/AnimNode_BoneDrivenController.h
+++ b/Engine/Source/Runtime/AnimGraphRuntime/Public/BoneControllers/AnimNode_BoneDrivenController.h
@@ -62,22 +62,22 @@ struct ANIMGRAPHRUNTIME_API FAnimNode_BoneDrivenController : public FAnimNode_Sk
 	// Minimum limit of the input value (mapped to RemappedMin, only used when limiting the source range)
 	// If this is rotation, the unit is radian
 	UPROPERTY(EditAnywhere, Category=Mapping, meta=(EditCondition=bUseRange, DisplayName="Source Range Min"))
-	float RangeMin;
+	double RangeMin;
 
 	// Maximum limit of the input value (mapped to RemappedMax, only used when limiting the source range)
 	// If this is rotation, the unit is radian
 	UPROPERTY(EditAnywhere, Category=Mapping, meta=(EditCondition=bUseRange, DisplayName="Source Range Max"))
-	float RangeMax;
+	double RangeMax;
 
 	// Minimum value to apply to the destination (remapped from the input range)
 	// If this is rotation, the unit is radian
 	UPROPERTY(EditAnywhere, Category=Mapping, meta=(EditCondition=bUseRange, DisplayName="Mapped Range Min"))
-	float RemappedMin;
+	double RemappedMin;
 
 	// Maximum value to apply to the destination (remapped from the input range)
 	// If this is rotation, the unit is radian
 	UPROPERTY(EditAnywhere, Category = Mapping, meta = (EditCondition = bUseRange, DisplayName="Mapped Range Max"))
-	float RemappedMax;
+	double RemappedMax;
 
 	/** Name of Morph Target to drive using the source attribute */
 	UPROPERTY(EditAnywhere, Category = "Destination (driven)")
@@ -170,6 +170,6 @@ protected:
 	virtual void InitializeBoneReferences(const FBoneContainer& RequiredBones) override;
 
 	/** Extracts the value used to drive the target bone or parameter */
-	const float ExtractSourceValue(const FTransform& InCurrentBoneTransform, const FTransform& InRefPoseBoneTransform);
+	const double ExtractSourceValue(const FTransform& InCurrentBoneTransform, const FTransform& InRefPoseBoneTransform);
 	// End of FAnimNode_SkeletalControlBase protected interface
 };
diff --git a/Engine/Source/Runtime/AnimGraphRuntime/Public/BoneControllers/AnimNode_LegIK.h b/Engine/Source/Runtime/AnimGraphRuntime/Public/BoneControllers/AnimNode_LegIK.h
--- a/Engine/Source/Runtime/AnimGraphRuntime/Public/BoneControllers/AnimNode_LegIK.h
+++ b/Engine/Source/Runtime/AnimGraphRuntime/Public/BoneControllers/AnimNode_LegIK.h
@@ -20,7 +20,7 @@ struct FIKChainLink
 
 public:
 	FVector Location;
-	float Length;
+	double Length;
 	FVector LinkAxisZ;
 	FVector RealBendDir;
 	FVector BaseBendDir;
@@ -28,14 +28,14 @@ public:
 
 	FIKChainLink()
 		: Location(FVector::ZeroVector)
-		, Length(0.f)
+		, Length(0.0)
 		, LinkAxisZ(FVector::ZeroVector)
 		, RealBendDir(FVector::ZeroVector)
 		, BaseBendDir(FVector::ZeroVector)
 		, BoneName(NAME_None)
 	{}
 
-	FIKChainLink(FVector InLocation, float InLength)
+	FIKChainLink(FVector InLocation, double InLength)
 		: Location(InLocation)
 		, Length(InLength)
 		, LinkAxisZ(FVector::ZeroVector)
@@ -56,7 +56,7 @@ public:
 
 private:
 	FAnimInstanceProxy* MyAnimInstanceProxy;
-	float MaximumReach;
+	double MaximumReach;
 	int32 NumLinks;
 	FVector HingeRotationAxis;
 	bool bEnableRotationLimit;
@@ -65,7 +65,7 @@ private:
 public:
 	FIKChain()
 		: MyAnimInstanceProxy(nullptr)
-		, MaximumReach(0.f)
+		, MaximumReach(0.0)
 		, NumLinks(INDEX_NONE)
 		, HingeRotationAxis(FVector::ZeroVector)
 		, bEnableRotationLimit(false)
@@ -73,9 +73,9 @@ public:
 	{}
 
 	void InitializeFromLegData(FAnimLegIKData& InLegData, FAnimInstanceProxy* InAnimInstanceProxy);
-	void ReachTarget(const FVector& InTargetLocation, float InReachPrecision, int32 InMaxIterations);
+	void ReachTarget(const FVector& InTargetLocation, double InReachPrecision, int32 InMaxIterations);
 
-	float GetMaximumReach() const
+	double GetMaximumReach() const
 	{
 		return MaximumReach;
 	}
@@ -83,7 +83,7 @@ public:
 private:
 	void OrientAllLinksToDirection(const FVector& InDirection);
 	void SolveTwoBoneIK(const FVector& InTargetLocation);
-	void SolveFABRIK(const FVector& InTargetLocation, float InReachPrecision, int32 InMaxIterations);
+	void SolveFABRIK(const FVector& InTargetLocation, double InReachPrecision, int32 InMaxIterations);
 
 	static void FABRIK_ForwardReach(const FVector& InTargetLocation, FIKChain& IKChain);
 	static void FABRIK_BackwardReach(const FVector& InRootTargetLocation, FIKChain& IKChain);
diff --git a/Engine/Source/Runtime/AnimGraphRuntime/Public/BoneControllers/AnimNode_RigidBody.h b/Engine/Source/Runtime/AnimGraphRuntime/Public/BoneControllers/AnimNode_RigidBody.h
--- a/Engine/Source/Runtime/AnimGraphRuntime/Public/BoneControllers/AnimNode_RigidBody.h
+++ b/Engine/Source/Runtime/AnimGraphRuntime/Public/BoneControllers/AnimNode_RigidBody.h
@@ -375,8 +375,8 @@ public:
 
 private:
 
-	float WorldTimeSeconds;
-	float LastEvalTimeSeconds;
+	double WorldTimeSeconds;
+	double LastEvalTimeSeconds;
 
 	float AccumulatedDeltaTime;
 	float AnimPhysicsMinDeltaTime;
@@ -476,7 +476,7 @@ private:
 
 	FVector WorldSpaceGravity;
 
-	float TotalMass;
+	double TotalMass;
 
 	// Bounds used to gather world objects copied into the simulation
 	FSphere CachedBounds;
diff --git a/Engine/Source/Runtime/AnimGraphRuntime/Public/BoneControllers/AnimNode_SpringBone.h b/Engine/Source/Runtime/AnimGraphRuntime/Public/BoneControllers/AnimNode_SpringBone.h
--- a/Engine/Source/Runtime/AnimGraphRuntime/Public/BoneControllers/AnimNode_SpringBone.h
+++ b/Engine/Source/Runtime/AnimGraphRuntime/Public/BoneControllers/AnimNode_SpringBone.h
@@ -27,19 +27,19 @@ struct ANIMGRAPHRUNTIME_API FAnimNode_SpringBone : public FAnimNode_SkeletalCont
 
 	/** If bLimitDisplacement is true, this indicates how long a bone can stretch beyond its length in the ref-pose. */
 	UPROPERTY(EditAnywhere, BlueprintReadWrite, Category=Spring, meta=(EditCondition="bLimitDisplacement"))
-	float MaxDisplacement;
+	double MaxDisplacement;
 
 	/** Stiffness of spring */
 	UPROPERTY(EditAnywhere, Category=Spring, meta = (PinHiddenByDefault))
-	float SpringStiffness;
+	double SpringStiffness;
 
 	/** Damping of spring */
 	UPROPERTY(EditAnywhere, Category=Spring, meta = (PinHiddenByDefault))
-	float SpringDamping;
+	double SpringDamping;
 
 	/** If spring stretches more than this, reset it. Useful for catching teleports etc */
 	UPROPERTY(EditAnywhere, Category=Spring)
-	float ErrorResetThresh;
+	double ErrorResetThresh;
 
 	/** World-space location of the bone. */
 	FVector BoneLocation;
diff --git a/Engine/Source/Runtime/AnimGraphRuntime/Public/BoneControllers/AnimNode_TwoBoneIK.h b/Engine/Source/Runtime/AnimGraphRuntime/Public/BoneControllers/AnimNode_TwoBoneIK.h
--- a/Engine/Source/Runtime/AnimGraphRuntime/Public/BoneControllers/AnimNode_TwoBoneIK.h
+++ b/Engine/Source/Runtime/AnimGraphRuntime/Public/BoneControllers/AnimNode_TwoBoneIK.h
@@ -27,11 +27,11 @@ struct ANIMGRAPHRUNTIME_API FAnimNode_TwoBoneIK : public FAnimNode_SkeletalContr
 
 	/** Limits to use if stretching is allowed. This value determines when to start stretch. For example, 0.9 means once it reaches 90% of the whole length of the limb, it will start apply. */
 	UPROPERTY(EditAnywhere, Category=IK, meta = (editcondition = "bAllowStretching", ClampMin = "0.0", UIMin = "0.0"))
-	float StartStretchRatio;
+	double StartStretchRatio;
 
 	/** Limits to use if stretching is allowed. This value determins what is the max stretch scale. For example, 1.5 means it will stretch until 150 % of the whole length of the limb.*/
 	UPROPERTY(EditAnywhere, Category= IK, meta = (editcondition = "bAllowStretching", ClampMin = "0.0", UIMin = "0.0"))
-	float MaxStretchScale;
+	double MaxStretchScale;
 
 #if WITH_EDITORONLY_DATA
 	/** Limits to use if stretching is allowed - old property DEPRECATED */
diff --git a/Engine/Source/Runtime/AnimGraphRuntime/Public/RBF/RBFInterpolator.h b/Engine/Source/Runtime/AnimGraphRuntime/Public/RBF/RBFInterpolator.h
--- a/Engine/Source/Runtime/AnimGraphRuntime/Public/RBF/RBFInterpolator.h
+++ b/Engine/Source/Runtime/AnimGraphRuntime/Public/RBF/RBFInterpolator.h
@@ -16,13 +16,13 @@
 namespace RBFDistanceMetric
 {
 	/* Returns the Euclidean (L2) distance between two coordinate vectors. */
-	static inline float Euclidean(const FVector& A, const FVector& B)
+	static inline double Euclidean(const FVector& A, const FVector& B)
 	{
 		return FVector::Distance(A, B);
 	}
 
 	/* Returns the Manhattan (L1), or Taxi-cab distance between two coordinate vectors. */
-	static inline float Manhattan(const FVector& A, const FVector& B)
+	static inline double Manhattan(const FVector& A, const FVector& B)
 	{
 		FVector AbsDiff = (A - B).GetAbs();
 		return AbsDiff.X + AbsDiff.Y + AbsDiff.Z;
@@ -30,7 +30,7 @@ namespace RBFDistanceMetric
 
 	/* Returns the arc length between two unit vectors (i.e. the distance between two
 	   points on a unit sphere, traveling along the surface of the sphere) */
-	static inline float ArcLength(const FVector& A, const FVector B)
+	static inline double ArcLength(const FVector& A, const FVector B)
 	{
 		return FMath::Acos(A.GetSafeNormal() | B.GetSafeNormal());
 	}
@@ -39,7 +39,7 @@ namespace RBFDistanceMetric
 	/* Returns a straight-up Euclidean distance between two rotation values expressed
 	   in radians.
 	*/
-	static inline float Euclidean(const FRotator& A, const FRotator& B)
+	static inline double Euclidean(const FRotator& A, const FRotator& B)
 	{
 		return Euclidean(FVector(FMath::DegreesToRadians(A.Roll),
 								 FMath::DegreesToRadians(A.Pitch),
@@ -51,21 +51,21 @@ namespace RBFDistanceMetric
 
 	/* Returns the arc-length distance, on a unit sphere, between two rotation vectors.
 	*/
-	static inline float ArcLength(const FRotator& A, const FRotator& B)
+	static inline double ArcLength(const FRotator& A, const FRotator& B)
 	{
 		return FMath::Acos(A.Vector() | B.Vector());
 	}
 
 	/* Returns the Euclidean (L2) distance between two quaternion values expressed.
 	*/
-	static inline float Euclidean(const FQuat& A, const FQuat& B)
+	static inline double Euclidean(const FQuat& A, const FQuat& B)
 	{
 		return (A - B).Size();
 	}
 
 	/* Returns the arc-length distance, on a unit sphere, between two quaternions.
 	*/
-	static inline float ArcLength(const FQuat& A, const FQuat& B)
+	static inline double ArcLength(const FQuat& A, const FQuat& B)
 	{
 		return A.GetNormalized().AngularDistance(B.GetNormalized());
 	}
@@ -73,7 +73,7 @@ namespace RBFDistanceMetric
 	/* Returns the swing arc length distance between two quaternions, using a specific 
 	   twist basis vector as reference.
 	*/
-	static inline float SwingAngle(const FQuat& A, const FQuat& B, const FVector& TwistAxis)
+	static inline double SwingAngle(const FQuat& A, const FQuat& B, const FVector& TwistAxis)
 	{
 		FQuat ASwing, BSwing, DummyTwist;
 		A.ToSwingTwist(TwistAxis, ASwing, DummyTwist);
@@ -84,7 +84,7 @@ namespace RBFDistanceMetric
 	/* Returns the twist arc length distance between two quaternions, using a specific 
 	   twist basis vector as reference.
 	*/
-	static inline float TwistAngle(const FQuat& A, const FQuat& B, const FVector& TwistAxis)
+	static inline double TwistAngle(const FQuat& A, const FQuat& B, const FVector& TwistAxis)
 	{
 		return FMath::Abs(A.GetTwistAngle(TwistAxis) - B.GetTwistAngle(TwistAxis));
 	}
diff --git a/Engine/Source/Runtime/AnimationCore/Private/FABRIK.cpp b/Engine/Source/Runtime/AnimationCore/Private/FABRIK.cpp
--- a/Engine/Source/Runtime/AnimationCore/Private/FABRIK.cpp
+++ b/Engine/Source/Runtime/AnimationCore/Private/FABRIK.cpp
@@ -10,7 +10,7 @@ namespace AnimationCore
 	// Implementation of the FABRIK IK Algorithm
 	// Please see http://andreasaristidou.com/publications/FABRIK.pdf for more details
 
-	bool SolveFabrik(TArray<FFABRIKChainLink>& InOutChain, const FVector& TargetPosition, float MaximumReach, float Precision, int32 MaxIterations)
+	bool SolveFabrik(TArray<FFABRIKChainLink>& InOutChain, const FVector& TargetPosition, double MaximumReach, double Precision, int32 MaxIterations)
 	{
 		bool bBoneLocationUpdated = false;
 		double const RootToTargetDistSq = FVector::DistSquared(InOutChain[0].Position, TargetPosition);
diff --git a/Engine/Source/Runtime/AnimationCore/Public/FABRIK.h b/Engine/Source/Runtime/AnimationCore/Public/FABRIK.h
--- a/Engine/Source/Runtime/AnimationCore/Public/FABRIK.h
+++ b/Engine/Source/Runtime/AnimationCore/Public/FABRIK.h
@@ -30,7 +30,7 @@ public:
 	FVector Position;
 
 	/** Distance to its parent link. */
-	float Length;
+	double Length;
 
 	/** Bone Index in SkeletalMesh */
 	int32 BoneIndex;
@@ -47,23 +47,23 @@ public:
 
 	FFABRIKChainLink()
 		: Position(FVector::ZeroVector)
-		, Length(0.f)
+		, Length(0.0)
 		, BoneIndex(INDEX_NONE)
 		, TransformIndex(INDEX_NONE)
-		, DefaultDirToParent(FVector(-1.f, 0.f, 0.f))
+		, DefaultDirToParent(FVector(-1.0, 0.0, 0.0))
 	{
 	}
 
-	FFABRIKChainLink(const FVector& InPosition, const float InLength, const FCompactPoseBoneIndex& InBoneIndex, const int32& InTransformIndex)
+	FFABRIKChainLink(const FVector& InPosition, const double InLength, const FCompactPoseBoneIndex& InBoneIndex, const int32& InTransformIndex)
 		: Position(InPosition)
 		, Length(InLength)
 		, BoneIndex(InBoneIndex.GetInt())
 		, TransformIndex(InTransformIndex)
-		, DefaultDirToParent(FVector(-1.f, 0.f, 0.f))
+		, DefaultDirToParent(FVector(-1.0, 0.0, 0.0))
 	{
 	}
 
-	FFABRIKChainLink(const FVector& InPosition, const float InLength, const FCompactPoseBoneIndex& InBoneIndex, const int32& InTransformIndex, const FVector& InDefaultDirToParent)
+	FFABRIKChainLink(const FVector& InPosition, const double InLength, const FCompactPoseBoneIndex& InBoneIndex, const int32& InTransformIndex, const FVector& InDefaultDirToParent)
 		: Position(InPosition)
 		, Length(InLength)
 		, BoneIndex(InBoneIndex.GetInt())
@@ -72,12 +72,12 @@ public:
 	{
 	}
 
-	FFABRIKChainLink(const FVector& InPosition, const float InLength, const int32 InBoneIndex, const int32 InTransformIndex)
+	FFABRIKChainLink(const FVector& InPosition, const double InLength, const int32 InBoneIndex, const int32 InTransformIndex)
 		: Position(InPosition)
 		, Length(InLength)
 		, BoneIndex(InBoneIndex)
 		, TransformIndex(InTransformIndex)
-		, DefaultDirToParent(FVector(-1.f, 0.f, 0.f))
+		, DefaultDirToParent(FVector(-1.0, 0.0, 0.0))
 	{
 	}
 };
@@ -97,5 +97,5 @@ namespace AnimationCore
 	*
 	* @return  true if modified. False if not.
 	*/
-	ANIMATIONCORE_API bool SolveFabrik(TArray<FFABRIKChainLink>& InOutChain, const FVector& TargetPosition, float MaximumReach, float Precision, int32 MaxIteration);
+	ANIMATIONCORE_API bool SolveFabrik(TArray<FFABRIKChainLink>& InOutChain, const FVector& TargetPosition, double MaximumReach, double Precision, int32 MaxIteration);
 };
