Skeleton compatibility improvements

Skeleton compatibility is now bi-directional. Specifying a compatible skeleton A -> B now implies B -> A.
Skeleton compatibility is now an editor-only concern. The runtime will attempt to do the 'best it can' via name -> name mappings. Only the editor will prevent assigning incompatible skeletons in (e.g.) asset pickers etc.
Skeleton compatibility checks in editor can now be disabled in the editor preferences (and each asset picker now has a checkbox option in its view settings that allows for quick access to this).

Moves FSkeletonRemapping to its own file (which is now private).
Skeleton remappings are now generated on demand on worker threads just before animation decompression and stored in a registry, guarded by FRWScopeLock for thread-safety.

Fixed some anim BP compiler edge cases where asset references on pins were not getting preloaded correctly, causing skeletons to be erroneously reported as missing.

Exposed the current asset registry filter in SAssetView so that menu extensions can access it (and use it to provide context)

diff --git a/Engine/Source/Editor/Persona/Private/SAssetFamilyShortcutBar.cpp b/Engine/Source/Editor/Persona/Private/SAssetFamilyShortcutBar.cpp
--- a/Engine/Source/Editor/Persona/Private/SAssetFamilyShortcutBar.cpp
+++ b/Engine/Source/Editor/Persona/Private/SAssetFamilyShortcutBar.cpp
@@ -519,6 +519,9 @@ private:
 
 	/** Cache the package's dirty state */
 	bool bPackageDirty;
+
+	/** Flag for handling deferred refreshes */
+	bool bDeferAssetAdded;
 };
 
 void SAssetFamilyShortcutBar::Construct(const FArguments& InArgs, const TSharedRef<FWorkflowCentricApplication>& InHostingApp, const TSharedRef<IAssetFamily>& InAssetFamily)
diff --git a/Engine/Source/Editor/UnrealEd/Private/PreferenceStubs.cpp b/Engine/Source/Editor/UnrealEd/Private/PreferenceStubs.cpp
--- a/Engine/Source/Editor/UnrealEd/Private/PreferenceStubs.cpp
+++ b/Engine/Source/Editor/UnrealEd/Private/PreferenceStubs.cpp
@@ -12,6 +12,7 @@
 #include "Preferences/MaterialStatsOptions.h"
 #include "FrameNumberDisplayFormat.h"
 #include "RHI.h"
+#include "Animation/Skeleton.h"
 
 // @todo find a better place for all of this, preferably in the appropriate modules
 // though this would require the classes to be relocated as well
diff --git a/Engine/Source/Runtime/AnimGraphRuntime/Private/AnimNodes/AnimNode_PoseBlendNode.cpp b/Engine/Source/Runtime/AnimGraphRuntime/Private/AnimNodes/AnimNode_PoseBlendNode.cpp
--- a/Engine/Source/Runtime/AnimGraphRuntime/Private/AnimNodes/AnimNode_PoseBlendNode.cpp
+++ b/Engine/Source/Runtime/AnimGraphRuntime/Private/AnimNodes/AnimNode_PoseBlendNode.cpp
@@ -47,7 +47,7 @@ void FAnimNode_PoseBlendNode::Evaluate_AnyThread(FPoseContext& Output)
 	bool bValidPose = false;
 
 	const UPoseAsset* CachedPoseAsset = CurrentPoseAsset.Get();
-	if (PoseAsset && PoseExtractContext.PoseCurves.Num() > 0 && PoseAsset->GetSkeleton() != nullptr)
+	if (CachedPoseAsset && PoseExtractContext.PoseCurves.Num() > 0 && CachedPoseAsset->GetSkeleton() != nullptr)
 	{
 		FPoseContext CurrentPose(Output);
 		// only give pose curve, we don't set any more curve here
diff --git a/Engine/Source/Runtime/AnimGraphRuntime/Private/AnimNodes/AnimNode_PoseDriver.cpp b/Engine/Source/Runtime/AnimGraphRuntime/Private/AnimNodes/AnimNode_PoseDriver.cpp
--- a/Engine/Source/Runtime/AnimGraphRuntime/Private/AnimNodes/AnimNode_PoseDriver.cpp
+++ b/Engine/Source/Runtime/AnimGraphRuntime/Private/AnimNodes/AnimNode_PoseDriver.cpp
@@ -377,7 +377,7 @@ void FAnimNode_PoseDriver::Evaluate_AnyThread(FPoseContext& Output)
 			const UPoseAsset* CachedPoseAsset = CurrentPoseAsset.Get();
 			if (DriveOutput == EPoseDriverOutput::DrivePoses &&
 				CachedPoseAsset &&
-				CurrentPoseAsset->GetSkeleton() != nullptr)
+				CachedPoseAsset->GetSkeleton() != nullptr)
 			{
 				FPoseContext CurrentPose(Output);
 
diff --git a/Engine/Source/Runtime/Engine/Classes/Animation/Skeleton.h b/Engine/Source/Runtime/Engine/Classes/Animation/Skeleton.h
--- a/Engine/Source/Runtime/Engine/Classes/Animation/Skeleton.h
+++ b/Engine/Source/Runtime/Engine/Classes/Animation/Skeleton.h
@@ -35,6 +35,7 @@ struct FAssetData;
 enum class EPackageReloadPhase : uint8;
 class USkeleton;
 typedef SmartName::UID_Type SkeletonAnimCurveUID;
+class USkeleton;
 
 // Delegate used to control global skeleton compatibility
 DECLARE_DELEGATE_RetVal(bool, FAreAllSkeletonsCompatible);
diff --git a/Engine/Source/Runtime/Engine/Private/Animation/AnimCompressionTypes.cpp b/Engine/Source/Runtime/Engine/Private/Animation/AnimCompressionTypes.cpp
--- a/Engine/Source/Runtime/Engine/Private/Animation/AnimCompressionTypes.cpp
+++ b/Engine/Source/Runtime/Engine/Private/Animation/AnimCompressionTypes.cpp
@@ -16,6 +16,8 @@
 #include "Misc/FileHelper.h"
 #include "Serialization/JsonWriter.h"
 #include "UObject/LinkerLoad.h"
+#include "Animation/SkeletonRemapping.h"
+#include "Animation/SkeletonRemappingRegistry.h"
 
 static FString GCompressionJsonOutput;
 static FAutoConsoleVariableRef CVarCompressionJsonOutput(
diff --git a/Engine/Source/Runtime/Engine/Private/Animation/AnimSequenceHelpers.cpp b/Engine/Source/Runtime/Engine/Private/Animation/AnimSequenceHelpers.cpp
--- a/Engine/Source/Runtime/Engine/Private/Animation/AnimSequenceHelpers.cpp
+++ b/Engine/Source/Runtime/Engine/Private/Animation/AnimSequenceHelpers.cpp
@@ -17,6 +17,9 @@
 #include "Animation/AnimNotifies/AnimNotify.h"
 #include "Animation/AnimNotifies/AnimNotifyState.h"
 #include "Misc/MessageDialog.h"
+#include "Animation/AnimSequenceHelpers.h"
+#include "Animation/SkeletonRemapping.h"
+#include "Animation/SkeletonRemappingRegistry.h"
 #include "Animation/AttributesRuntime.h"
 
 #define LOCTEXT_NAMESPACE "AnimSequenceHelpers"
