Chaos - fix RBAN SimulationSpace force transfer when using Bone-Space mode and scaling to zero

diff --git a/Engine/Source/Runtime/AnimGraphRuntime/Private/BoneControllers/AnimNode_RigidBody.cpp b/Engine/Source/Runtime/AnimGraphRuntime/Private/BoneControllers/AnimNode_RigidBody.cpp
--- a/Engine/Source/Runtime/AnimGraphRuntime/Private/BoneControllers/AnimNode_RigidBody.cpp
+++ b/Engine/Source/Runtime/AnimGraphRuntime/Private/BoneControllers/AnimNode_RigidBody.cpp
@@ -423,6 +423,12 @@ void FAnimNode_RigidBody::CalculateSimulationSpace(
 	SpaceLinearAcc = FVector::ZeroVector;
 	SpaceAngularAcc = FVector::ZeroVector;
 
+	// The simulation scale does not change - we scale the inputs and outputs instead.
+	// This means we do not support phantom forces resulting from scale changes, but that's ok.
+	// NOTE: If we don't clear the scale, rapid scaling to zero can introduce large phantom forces
+	// leading to major instability in the simulation
+	SpaceTransform.SetScale3D(FVector::One());
+
 	// If the system is disabled, nothing else to do
 	if ((Settings.WorldAlpha == 0.0f) || (Dt < SMALL_NUMBER))
 	{
