Fixed graph based sync not working correctly when a sync point does not specify a sync group name

Even with a name of 'None', a sync point (FAnimSyncGroupScope - such as that on a blend space graph) would forward for syncing to the proxy, bypassing any other valid rootwards sync scopes.
This change prevents adding the sync messages completely if their name is 'None', ensuring that outer scopes correctly receive the tick record.

diff --git a/Engine/Source/Runtime/AnimGraphRuntime/Private/AnimNodes/AnimNode_BlendSpaceGraphBase.cpp b/Engine/Source/Runtime/AnimGraphRuntime/Private/AnimNodes/AnimNode_BlendSpaceGraphBase.cpp
--- a/Engine/Source/Runtime/AnimGraphRuntime/Private/AnimNodes/AnimNode_BlendSpaceGraphBase.cpp
+++ b/Engine/Source/Runtime/AnimGraphRuntime/Private/AnimNodes/AnimNode_BlendSpaceGraphBase.cpp
@@ -82,7 +82,8 @@ void FAnimNode_BlendSpaceGraphBase::Update_AnyThread(const FAnimationUpdateConte
 	DECLARE_SCOPE_HIERARCHICAL_COUNTER_ANIMNODE(Update_AnyThread)
 	GetEvaluateGraphExposedInputs().Execute(Context);
 
-	UE::Anim::TScopedGraphMessage<UE::Anim::FAnimSyncGroupScope> Message(Context, Context, GroupName, GroupRole);
+	const bool bApplySyncing = GroupName != NAME_None;
+	UE::Anim::TOptionalScopedGraphMessage<UE::Anim::FAnimSyncGroupScope> Message(bApplySyncing, Context, Context, GroupName, GroupRole);
 
 	UpdateInternal(Context);
 }
diff --git a/Engine/Source/Runtime/AnimGraphRuntime/Private/AnimNodes/AnimNode_Sync.cpp b/Engine/Source/Runtime/AnimGraphRuntime/Private/AnimNodes/AnimNode_Sync.cpp
--- a/Engine/Source/Runtime/AnimGraphRuntime/Private/AnimNodes/AnimNode_Sync.cpp
+++ b/Engine/Source/Runtime/AnimGraphRuntime/Private/AnimNodes/AnimNode_Sync.cpp
@@ -12,7 +12,8 @@ void FAnimNode_Sync::Initialize_AnyThread(const FAnimationInitializeContext& Con
 
 void FAnimNode_Sync::Update_AnyThread(const FAnimationUpdateContext& Context)
 {
-	UE::Anim::TScopedGraphMessage<UE::Anim::FAnimSyncGroupScope> Message(Context, Context, GroupName, GroupRole);
+	const bool bApplySyncing = GroupName != NAME_None;
+	UE::Anim::TOptionalScopedGraphMessage<UE::Anim::FAnimSyncGroupScope> Message(bApplySyncing, Context, Context, GroupName, GroupRole);
 
 	Source.Update(Context);
 }
