Added validation and visual error functionality to AnimNode_SkeletalControlBase, which can be used to check referenced bones and show errors on the AnimGraph nodes that use them (so it is clear if the setup does not work due to a bone missing)

diff --git a/Engine/Source/Editor/AnimGraph/Private/AnimGraphNode_SkeletalControlBase.cpp b/Engine/Source/Editor/AnimGraph/Private/AnimGraphNode_SkeletalControlBase.cpp
--- a/Engine/Source/Editor/AnimGraph/Private/AnimGraphNode_SkeletalControlBase.cpp
+++ b/Engine/Source/Editor/AnimGraph/Private/AnimGraphNode_SkeletalControlBase.cpp
@@ -4,6 +4,7 @@
 #include "UnrealWidgetFwd.h"
 #include "AnimationGraphSchema.h"
 #include "Animation/AnimationSettings.h"
+#include "Animation/AnimInstance.h"
 #include "Components/SkeletalMeshComponent.h"
 #include "Engine/SkeletalMesh.h"
 #include "Kismet2/CompilerResultsLog.h"
@@ -455,6 +456,22 @@ void UAnimGraphNode_SkeletalControlBase::PostEditChangeProperty(struct FProperty
 	Super::PostEditChangeProperty(PropertyChangedEvent);
 }
 
+bool UAnimGraphNode_SkeletalControlBase::ShowVisualWarning() const
+{
+	const FAnimNode_SkeletalControlBase* DebuggedNode = GetDebuggedNode();
+
+	const bool bHasError = DebuggedNode ? DebuggedNode->HasValidationVisualWarnings() : false;
+
+	return bHasError;
+}
+
+FText UAnimGraphNode_SkeletalControlBase::GetVisualWarningTooltipText() const
+{
+	FAnimNode_SkeletalControlBase* DebuggedNode = GetDebuggedNode();
+
+	return DebuggedNode ? DebuggedNode->GetValidationVisualWarningMessage() : FText::GetEmpty();
+}
+
 void UAnimGraphNode_SkeletalControlBase::CustomizeDetails(IDetailLayoutBuilder& DetailBuilder)
 {
 	Super::CustomizeDetails(DetailBuilder);
@@ -498,4 +515,21 @@ void UAnimGraphNode_SkeletalControlBase::ValidateAnimNodePostCompile(FCompilerRe
 	}
 }
 
+FAnimNode_SkeletalControlBase* UAnimGraphNode_SkeletalControlBase::GetDebuggedNode() const
+{
+	if (const UObject* ObjectBeingDebugged = GetAnimBlueprint()->GetObjectBeingDebugged())
+	{
+		if (const UAnimInstance* InstanceBeingDebugged = Cast<const UAnimInstance>(ObjectBeingDebugged))
+		{
+			USkeletalMeshComponent* Component = InstanceBeingDebugged->GetSkelMeshComponent();
+			if (Component != nullptr && Component->GetAnimInstance() != nullptr)
+			{
+				return static_cast<FAnimNode_SkeletalControlBase*>(FindDebugAnimNode(Component));
+			}
+		}
+	}
+
+	return nullptr;
+}
+
 #undef LOCTEXT_NAMESPACE
diff --git a/Engine/Source/Editor/AnimGraph/Public/AnimGraphNode_SkeletalControlBase.h b/Engine/Source/Editor/AnimGraph/Public/AnimGraphNode_SkeletalControlBase.h
--- a/Engine/Source/Editor/AnimGraph/Public/AnimGraphNode_SkeletalControlBase.h
+++ b/Engine/Source/Editor/AnimGraph/Public/AnimGraphNode_SkeletalControlBase.h
@@ -29,6 +29,8 @@ public:
 	virtual FLinearColor GetNodeTitleColor() const override;
 	virtual FText GetTooltipText() const override;
 	virtual void PostEditChangeProperty(struct FPropertyChangedEvent& PropertyChangedEvent) override;
+	virtual bool ShowVisualWarning() const override;
+	virtual FText GetVisualWarningTooltipText() const override;
 	// End of UEdGraphNode interface
 
 	// UAnimGraphNode_Base interface
@@ -162,4 +164,6 @@ protected:
 	FVector ConvertWidgetLocation(const USkeletalMeshComponent* InSkelComp, FCSPose<FCompactHeapPose>& InMeshBases, const FName& BoneName, const FVector& InLocation, const EBoneControlSpace Space);
 
 	virtual const FAnimNode_SkeletalControlBase* GetNode() const PURE_VIRTUAL(UAnimGraphNode_SkeletalControlBase::GetNode, return nullptr;);
+
+	FAnimNode_SkeletalControlBase* GetDebuggedNode() const;
 };
diff --git a/Engine/Source/Editor/GraphEditor/Private/SGraphNode.cpp b/Engine/Source/Editor/GraphEditor/Private/SGraphNode.cpp
--- a/Engine/Source/Editor/GraphEditor/Private/SGraphNode.cpp
+++ b/Engine/Source/Editor/GraphEditor/Private/SGraphNode.cpp
@@ -765,6 +765,9 @@ void SGraphNode::UpdateErrorInfo()
 		ErrorColor = FLinearColor(0,0,0);
 		ErrorMsg.Empty();
 	}
+
+	VisualWarningMsg = FString(TEXT("WARNING!"));
+	VisualWarningColor = FAppStyle::GetColor("ErrorReporting.WarningBackgroundColor");
 }
 
 void SGraphNode::SetupErrorReporting()
@@ -783,6 +786,20 @@ void SGraphNode::SetupErrorReporting()
 		ErrorReporting = ErrorTextWidget;
 	}
 	ErrorReporting->SetError(ErrorMsg);
+
+	if (!VisualWarningReporting.IsValid())
+	{
+		TSharedPtr<SErrorText> ErrorTextWidget;
+
+		// generate widget
+		SAssignNew(ErrorTextWidget, SErrorText)
+			.Visibility(this, &SGraphNode::VisualWarningVisibility)
+			.BackgroundColor(this, &SGraphNode::GetVisualWarningColor)
+			.ToolTipText(this, &SGraphNode::GetVisualWarningMsgToolTip);
+
+		VisualWarningReporting = ErrorTextWidget;
+	}
+	VisualWarningReporting->SetError(VisualWarningMsg);
 }
 
 TSharedRef<SWidget> SGraphNode::CreateTitleWidget(TSharedPtr<SNodeTitle> NodeTitle)
@@ -982,6 +999,12 @@ void SGraphNode::UpdateGraphNode()
 			ErrorReporting->AsWidget()
 		];
 
+	InnerVerticalBox->AddSlot()
+		.AutoHeight()
+		.Padding(Settings->GetNonPinNodeBodyPadding())
+		[
+			VisualWarningReporting->AsWidget()
+		];
 
 
 	this->GetOrAddSlot( ENodeZone::Center )
@@ -1475,6 +1498,30 @@ FText SGraphNode::GetErrorMsgToolTip() const
 	return Result;
 }
 
+EVisibility SGraphNode::VisualWarningVisibility() const
+{
+	const bool bShowVisualWarning = GraphNode && GraphNode->ShowVisualWarning();
+	return bShowVisualWarning ? EVisibility::Visible : EVisibility::Collapsed;
+}
+
+/* Helper function to set the error color for the node */
+FSlateColor SGraphNode::GetVisualWarningColor() const
+{
+	return VisualWarningColor;
+}
+
+FText SGraphNode::GetVisualWarningMsgToolTip() const
+{
+	FText Result = FText::GetEmpty();
+	if (GraphNode != nullptr)
+	{
+		Result = GraphNode->GetVisualWarningTooltipText();
+	}
+
+	return Result;
+}
+
+
 bool SGraphNode::IsNameReadOnly() const
 {
 	return (!GraphNode->GetCanRenameNode() || !IsNodeEditable());
diff --git a/Engine/Source/Editor/GraphEditor/Public/SGraphNode.h b/Engine/Source/Editor/GraphEditor/Public/SGraphNode.h
--- a/Engine/Source/Editor/GraphEditor/Public/SGraphNode.h
+++ b/Engine/Source/Editor/GraphEditor/Public/SGraphNode.h
@@ -330,6 +330,15 @@ protected:
 	/* Helper function to set the error color for the node */
 	FSlateColor GetErrorColor() const	{return ErrorColor;}
 
+	/** Controls wether to show or not the visual warning message */
+	EVisibility VisualWarningVisibility() const;
+
+	/** Function to get a visual warning description string(visual warnings does not break the build) */
+	FText GetVisualWarningMsgToolTip() const;
+
+	/* Helper function to set the error color for the node */
+	FSlateColor GetVisualWarningColor() const;
+
 	/** Helper function to get any error text for the node */
 	FString GetErrorMessage() const {return ErrorMsg;}
 
@@ -424,6 +433,8 @@ protected:
 	TSharedPtr<SInlineEditableTextBlock> InlineEditableText;
 	/** Error handling widget */
 	TSharedPtr<class IErrorReportingWidget> ErrorReporting;
+	/** Visual Warning handling widget */
+	TSharedPtr<class IErrorReportingWidget> VisualWarningReporting;
 
 	FCurveSequence SpawnAnim;
 	FCurveHandle ZoomCurve;
@@ -449,6 +460,11 @@ protected:
 	FString ErrorMsg;
 	/** Used to set the error color */
 	FSlateColor ErrorColor;
+	/** Used to report visual warnings on the node (does not break build) */
+	FString VisualWarningMsg;
+	/** Used to set the soft error color */
+	FSlateColor VisualWarningColor;
+
 
 	/** Caches true position of node */
 	FVector2D CachedUnscaledPosition;
diff --git a/Engine/Source/Runtime/AnimGraphRuntime/Private/BoneControllers/AnimNode_RigidBody.cpp b/Engine/Source/Runtime/AnimGraphRuntime/Private/BoneControllers/AnimNode_RigidBody.cpp
--- a/Engine/Source/Runtime/AnimGraphRuntime/Private/BoneControllers/AnimNode_RigidBody.cpp
+++ b/Engine/Source/Runtime/AnimGraphRuntime/Private/BoneControllers/AnimNode_RigidBody.cpp
@@ -1773,17 +1773,9 @@ void FAnimNode_RigidBody::InitializeBoneReferences(const FBoneContainer& Require
 		BaseBoneRef.BoneName = RefSkeleton.GetBoneName(0);
 	}
 
-	if (BaseBoneRef.BoneName != NAME_None)
-	{
-		BaseBoneRef.Initialize(RequiredBones);
-	}
-
-	if (!BaseBoneRef.HasValidSetup())
-	{
-		// If the user specified a simulation root that is not used by the skelmesh, issue a warning 
-		// (FAnimNode_RigidBody::IsValidToEvaluate will return false and the simulation will not run)
-		UE_LOG(LogRBAN, Log, TEXT("FAnimNode_RigidBody: RBAN Simulation Base Bone \'%s\' does not exist on SkeletalMesh %s."), *BaseBoneRef.BoneName.ToString(), *GetNameSafe(RequiredBones.GetSkeletalMeshAsset()));
-	}
+	// If the user specified a simulation root that is not used by the skelmesh, issue a warning 
+	// (FAnimNode_RigidBody::IsValidToEvaluate will return false and the simulation will not run)
+	InitializeAndValidateBoneRef(BaseBoneRef, RequiredBones);
 
 	bool bHasInvalidBoneReference = false;
 	for (int32 Index = 0; Index < NumRequiredBoneIndices; ++Index)
diff --git a/Engine/Source/Runtime/AnimGraphRuntime/Private/BoneControllers/AnimNode_SkeletalControlBase.cpp b/Engine/Source/Runtime/AnimGraphRuntime/Private/BoneControllers/AnimNode_SkeletalControlBase.cpp
--- a/Engine/Source/Runtime/AnimGraphRuntime/Private/BoneControllers/AnimNode_SkeletalControlBase.cpp
+++ b/Engine/Source/Runtime/AnimGraphRuntime/Private/BoneControllers/AnimNode_SkeletalControlBase.cpp
@@ -1,11 +1,21 @@
 // Copyright Epic Games, Inc. All Rights Reserved.
 
 #include "BoneControllers/AnimNode_SkeletalControlBase.h"
+#include "Engine/SkeletalMesh.h"
 #include "Animation/AnimInstanceProxy.h"
 #include "Engine/SkeletalMeshSocket.h"
 
 #include UE_INLINE_GENERATED_CPP_BY_NAME(AnimNode_SkeletalControlBase)
 
+#define LOCTEXT_NAMESPACE "AnimNode_SkeletalControlBase"
+
+#if UE_BUILD_SHIPPING || UE_BUILD_TEST
+	DECLARE_LOG_CATEGORY_EXTERN(LogSkeletalControlBase, Log, Warning);
+#else
+	DECLARE_LOG_CATEGORY_EXTERN(LogSkeletalControlBase, Log, All);
+#endif
+DEFINE_LOG_CATEGORY(LogSkeletalControlBase);
+
 /////////////////////////////////////////////////////
 // FAnimNode_SkeletalControlBase
 
@@ -22,6 +32,10 @@ void FAnimNode_SkeletalControlBase::Initialize_AnyThread(const FAnimationInitial
 
 void FAnimNode_SkeletalControlBase::CacheBones_AnyThread(const FAnimationCacheBonesContext& Context)
 {
+#if WITH_EDITOR
+	ClearValidationVisualWarnings();
+#endif
+
 	DECLARE_SCOPE_HIERARCHICAL_COUNTER_ANIMNODE(CacheBones_AnyThread)
 	FAnimNode_Base::CacheBones_AnyThread(Context);
 	InitializeBoneReferences(Context.AnimInstanceProxy->GetRequiredBones());
@@ -160,3 +174,65 @@ float FAnimNode_SkeletalControlBase::GetAlpha() const
 {
 	return ActualAlpha;
 }
+
+void FAnimNode_SkeletalControlBase::InitializeAndValidateBoneRef(FBoneReference& BoneRef, const FBoneContainer& RequiredBones)
+{
+	if (BoneRef.BoneName != NAME_None)
+	{
+		BoneRef.Initialize(RequiredBones);
+	}
+
+	if (!BoneRef.HasValidSetup())
+	{
+		const FText ErrorText = FText::Format(LOCTEXT("SkeletalControlBoneError", "Referenced Bone {0} does not exist on SkeletalMesh {1}."),
+			FText::AsCultureInvariant(BoneRef.BoneName.ToString()),
+			FText::AsCultureInvariant(GetNameSafe(RequiredBones.GetSkeletalMeshAsset())));
+
+#if WITH_EDITOR
+		AddValidationVisualWarning(ErrorText);
+#endif // WITH_EDITOR
+
+		// If the user specified a simulation root that is not used by the skelmesh, issue a warning 
+		UE_LOG(LogSkeletalControlBase, Log, TEXT("%s"), *ErrorText.ToString());
+	}
+}
+
+#if WITH_EDITOR
+void FAnimNode_SkeletalControlBase::AddBoneRefMissingVisualWarning(const FString& BoneName, const FString& SkeletalMeshName)
+{
+	const FText ErrorText = FText::Format(LOCTEXT("SkeletalControlBoneError", "Simulation Base Bone {0} does not exist on SkeletalMesh {1}."), FText::FromString(BoneName), FText::FromString(SkeletalMeshName));
+	
+	AddValidationVisualWarning(ErrorText);
+}
+
+void FAnimNode_SkeletalControlBase::AddValidationVisualWarning(FText ValidationVisualWarning)
+{
+#if WITH_EDITORONLY_DATA
+	if (ValidationVisualWarningMessage.IsEmpty())
+	{
+		ValidationVisualWarningMessage = ValidationVisualWarning;
+	}
+	else
+	{
+		ValidationVisualWarningMessage = FText::Format(FText::FromString(TEXT("{0}\n{1}")), ValidationVisualWarningMessage, ValidationVisualWarning);
+	}
+#endif
+}
+
+void FAnimNode_SkeletalControlBase::ClearValidationVisualWarnings()
+{
+	ValidationVisualWarningMessage = FText::GetEmpty();
+}
+
+bool FAnimNode_SkeletalControlBase::HasValidationVisualWarnings() const
+{
+	return ValidationVisualWarningMessage.IsEmpty() == false;
+}
+
+FText FAnimNode_SkeletalControlBase::GetValidationVisualWarningMessage() const
+{
+	return ValidationVisualWarningMessage;
+}
+#endif // WITH_EDITOR
+
+#undef LOCTEXT_NAMESPACE
diff --git a/Engine/Source/Runtime/AnimGraphRuntime/Public/BoneControllers/AnimNode_SkeletalControlBase.h b/Engine/Source/Runtime/AnimGraphRuntime/Public/BoneControllers/AnimNode_SkeletalControlBase.h
--- a/Engine/Source/Runtime/AnimGraphRuntime/Public/BoneControllers/AnimNode_SkeletalControlBase.h
+++ b/Engine/Source/Runtime/AnimGraphRuntime/Public/BoneControllers/AnimNode_SkeletalControlBase.h
@@ -90,7 +90,19 @@ public:
 
 	// Get the alpha of this node
 	float GetAlpha() const;
-	
+
+	void InitializeAndValidateBoneRef(FBoneReference& BoneRef, const FBoneContainer& RequiredBones);
+
+	// Visual warnings are shown on the node but not logged as an error for build system, use with care
+	// The warnigns are cleared at CacheBones_AnyThread and should be added during InitializeBoneReferences
+#if WITH_EDITOR
+	void AddBoneRefMissingVisualWarning(const FString& BoneName, const FString& SkeletalMeshName);
+	void AddValidationVisualWarning(FText ValidationVisualWarning);
+	bool HasValidationVisualWarnings() const;
+	void ClearValidationVisualWarnings();
+	FText GetValidationVisualWarningMessage() const;
+#endif
+
 protected:
 	// Interface for derived skeletal controls to implement
 	// use this function to update for skeletal control base
@@ -113,8 +125,14 @@ protected:
 
 	/** Allow base to add info to the node debug output */
 	void AddDebugNodeData(FString& OutDebugData);
-private:
 
+private:
 	// Resused bone transform array to avoid reallocating in skeletal controls
 	TArray<FBoneTransform> BoneTransforms;
+
+#if WITH_EDITORONLY_DATA
+	UPROPERTY(transient)
+	FText ValidationVisualWarningMessage;
+#endif
+
 };
diff --git a/Engine/Source/Runtime/Engine/Classes/EdGraph/EdGraphNode.h b/Engine/Source/Runtime/Engine/Classes/EdGraph/EdGraphNode.h
--- a/Engine/Source/Runtime/Engine/Classes/EdGraph/EdGraphNode.h
+++ b/Engine/Source/Runtime/Engine/Classes/EdGraph/EdGraphNode.h
@@ -682,6 +682,12 @@ public:
 		bHasCompilerMessage = false;
 	}
 
+	/** If true, this node whill show the Visual Warning message */
+	virtual bool ShowVisualWarning() const;
+
+	/** Visual Warning tooltip message to show */
+	virtual FText GetVisualWarningTooltipText() const;
+
 	/** Generate a unique pin name, trying to stick close to a passed in name */
 	virtual FName CreateUniquePinName(FName SourcePinName) const
 	{
diff --git a/Engine/Source/Runtime/Engine/Private/EdGraph/EdGraphNode.cpp b/Engine/Source/Runtime/Engine/Private/EdGraph/EdGraphNode.cpp
--- a/Engine/Source/Runtime/Engine/Private/EdGraph/EdGraphNode.cpp
+++ b/Engine/Source/Runtime/Engine/Private/EdGraph/EdGraphNode.cpp
@@ -492,6 +492,16 @@ void UEdGraphNode::SnapToGrid(uint32 GridSnapSize)
 	NodePosY = GridSnapSize * (NodePosY / GridSnapSize);
 }
 
+bool UEdGraphNode::ShowVisualWarning() const
+{
+	return false;
+}
+
+FText UEdGraphNode::GetVisualWarningTooltipText() const
+{
+	return FText();
+}
+
 class UEdGraph* UEdGraphNode::GetGraph() const
 {
 	UEdGraph* Graph = Cast<UEdGraph>(GetOuter());
