Fix for CIS LocText at AnimNode_SkeletalControlBase

diff --git a/Engine/Source/Runtime/AnimGraphRuntime/Private/BoneControllers/AnimNode_SkeletalControlBase.cpp b/Engine/Source/Runtime/AnimGraphRuntime/Private/BoneControllers/AnimNode_SkeletalControlBase.cpp
--- a/Engine/Source/Runtime/AnimGraphRuntime/Private/BoneControllers/AnimNode_SkeletalControlBase.cpp
+++ b/Engine/Source/Runtime/AnimGraphRuntime/Private/BoneControllers/AnimNode_SkeletalControlBase.cpp
@@ -200,7 +200,7 @@ void FAnimNode_SkeletalControlBase::InitializeAndValidateBoneRef(FBoneReference&
 #if WITH_EDITOR
 void FAnimNode_SkeletalControlBase::AddBoneRefMissingVisualWarning(const FString& BoneName, const FString& SkeletalMeshName)
 {
-	const FText ErrorText = FText::Format(LOCTEXT("SkeletalControlBoneError", "Simulation Base Bone {0} does not exist on SkeletalMesh {1}."), FText::FromString(BoneName), FText::FromString(SkeletalMeshName));
+	const FText ErrorText = FText::Format(LOCTEXT("SkeletalControlBoneVisualWarning", "Simulation Base Bone {0} does not exist on SkeletalMesh {1}."), FText::FromString(BoneName), FText::FromString(SkeletalMeshName));
 	
 	AddValidationVisualWarning(ErrorText);
 }
