Pose Driver: Fix a crash when trying to dereference with an invalid bone id.

diff --git a/Engine/Source/Runtime/AnimGraphRuntime/Private/AnimNodes/AnimNode_PoseDriver.cpp b/Engine/Source/Runtime/AnimGraphRuntime/Private/AnimNodes/AnimNode_PoseDriver.cpp
--- a/Engine/Source/Runtime/AnimGraphRuntime/Private/AnimNodes/AnimNode_PoseDriver.cpp
+++ b/Engine/Source/Runtime/AnimGraphRuntime/Private/AnimNodes/AnimNode_PoseDriver.cpp
@@ -204,10 +204,14 @@ void FAnimNode_PoseDriver::GetRBFTargets(TArray<FRBFTarget>& OutTargets, const F
 				FTransform RefBoneTransform = FTransform::Identity;
 				if (bEvalFromRefPose && BoneContainer)
 				{
-					const FCompactPoseBoneIndex CompactPoseIndex = SourceBones[SourceIdx].CachedCompactPoseIndex;
-					if (CompactPoseIndex < BoneContainer->GetCompactPoseNumBones())
+					const FBoneReference& BoneReference = SourceBones[SourceIdx];
+					if (BoneReference.HasValidSetup())
 					{
-						RefBoneTransform = BoneContainer->GetRefPoseTransform(CompactPoseIndex);
+						const FCompactPoseBoneIndex CompactPoseIndex = BoneReference.CachedCompactPoseIndex;
+						if (CompactPoseIndex < BoneContainer->GetCompactPoseNumBones())
+						{
+							RefBoneTransform = BoneContainer->GetRefPoseTransform(CompactPoseIndex);
+						}
 					}
 				}
 
