Add anim node library function for blend list base
Add an anim node function to reset blend lists

diff --git a/Engine/Source/Runtime/AnimGraphRuntime/Private/AnimNodes/AnimNode_BlendListBase.cpp b/Engine/Source/Runtime/AnimGraphRuntime/Private/AnimNodes/AnimNode_BlendListBase.cpp
--- a/Engine/Source/Runtime/AnimGraphRuntime/Private/AnimNodes/AnimNode_BlendListBase.cpp
+++ b/Engine/Source/Runtime/AnimGraphRuntime/Private/AnimNodes/AnimNode_BlendListBase.cpp
@@ -16,7 +16,6 @@ void FAnimNode_BlendListBase::Initialize_AnyThread(const FAnimationInitializeCon
 {
 	DECLARE_SCOPE_HIERARCHICAL_COUNTER_ANIMNODE(Initialize_AnyThread)
 	FAnimNode_Base::Initialize_AnyThread(Context);
-
 	const int32 NumPoses = BlendPose.Num();
 	const TArray<float>& CurrentBlendTimes = GetBlendTimes();
 	checkSlow(CurrentBlendTimes.Num() == NumPoses);
@@ -35,13 +34,20 @@ void FAnimNode_BlendListBase::Initialize_AnyThread(const FAnimationInitializeCon
 		}
 	}
 
+	Initialize();
+}
+
+void FAnimNode_BlendListBase::Initialize()
+{
 	UBlendProfile* CurrentBlendProfile = GetBlendProfile();
 
+	const int32 NumPoses = BlendPose.Num();
+
 	LastActiveChildIndex = INDEX_NONE;
 
 	EAlphaBlendOption CurrentBlendType = GetBlendType();
 	UCurveFloat* CurrentCustomBlendCurve = GetCustomBlendCurve();
-	for(int32 i = 0 ; i < PerBlendData.Num() ; ++i)
+	for (int32 i = 0; i < PerBlendData.Num(); ++i)
 	{
 		FAlphaBlend& Blend = PerBlendData[i].Blend;
 
@@ -56,7 +62,7 @@ void FAnimNode_BlendListBase::Initialize_AnyThread(const FAnimationInitializeCon
 	}
 	PerBlendData[0].Blend.SetAlpha(1.0f);
 
-	if(CurrentBlendProfile)
+	if (CurrentBlendProfile)
 	{
 		PerBlendData[0].StartAlpha = 1.0f;
 
@@ -64,7 +70,7 @@ void FAnimNode_BlendListBase::Initialize_AnyThread(const FAnimationInitializeCon
 		PerBoneSampleData.Empty(NumPoses);
 		PerBoneSampleData.AddZeroed(NumPoses);
 
-		for(int32 Idx = 0 ; Idx < NumPoses ; ++Idx)
+		for (int32 Idx = 0; Idx < NumPoses; ++Idx)
 		{
 			FBlendSampleData& SampleData = PerBoneSampleData[Idx];
 			SampleData.SampleDataIndex = Idx;
diff --git a/Engine/Source/Runtime/AnimGraphRuntime/Private/BlendListBaseLibrary.cpp b/Engine/Source/Runtime/AnimGraphRuntime/Private/BlendListBaseLibrary.cpp
new file mode 100644
--- /dev/null
+++ b/Engine/Source/Runtime/AnimGraphRuntime/Private/BlendListBaseLibrary.cpp
@@ -0,0 +1,25 @@
+// Copyright Epic Games, Inc. All Rights Reserved.
+
+#include "BlendListBaseLibrary.h"
+#include "AnimNodes/AnimNode_BlendListBase.h"
+
+#include UE_INLINE_GENERATED_CPP_BY_NAME(BlendListBaseLibrary)
+
+
+FBlendListBaseReference UBlendListBaseLibrary::ConvertToBlendListBase(const FAnimNodeReference& Node,
+	EAnimNodeReferenceConversionResult& Result)
+{
+	return FAnimNodeReference::ConvertToType<FBlendListBaseReference>(Node, Result);
+}
+
+void UBlendListBaseLibrary::ResetNode(const FBlendListBaseReference& BlendListBase)
+{
+	bool bResult = false;
+
+	BlendListBase.CallAnimNodeFunction<FAnimNode_BlendListBase>(
+		TEXT("ResetNode"),
+		[](FAnimNode_BlendListBase& InBlendListBase)
+	{
+		InBlendListBase.Initialize();
+	});
+}
diff --git a/Engine/Source/Runtime/AnimGraphRuntime/Public/AnimNodes/AnimNode_BlendListBase.h b/Engine/Source/Runtime/AnimGraphRuntime/Public/AnimNodes/AnimNode_BlendListBase.h
--- a/Engine/Source/Runtime/AnimGraphRuntime/Public/AnimNodes/AnimNode_BlendListBase.h
+++ b/Engine/Source/Runtime/AnimGraphRuntime/Public/AnimNodes/AnimNode_BlendListBase.h
@@ -117,4 +117,8 @@ public:
 protected:
 	virtual int32 GetActiveChildIndex() { return 0; }
 	virtual FString GetNodeName(FNodeDebugData& DebugData) { return DebugData.GetNodeName(this); }
+
+	void Initialize();
+
+	friend class UBlendListBaseLibrary;
 };
diff --git a/Engine/Source/Runtime/AnimGraphRuntime/Public/BlendListBaseLibrary.h b/Engine/Source/Runtime/AnimGraphRuntime/Public/BlendListBaseLibrary.h
new file mode 100644
--- /dev/null
+++ b/Engine/Source/Runtime/AnimGraphRuntime/Public/BlendListBaseLibrary.h
@@ -0,0 +1,37 @@
+// Copyright Epic Games, Inc. All Rights Reserved.
+
+#pragma once
+
+#include "CoreMinimal.h"
+
+#include "Animation/AnimExecutionContext.h"
+#include "Kismet/BlueprintFunctionLibrary.h"
+#include "Animation/AnimNodeReference.h"
+#include "BlendListBaseLibrary.generated.h"
+
+struct FAnimNode_BlendListBase;
+USTRUCT(BlueprintType)
+struct FBlendListBaseReference : public FAnimNodeReference
+{
+	GENERATED_BODY()
+	typedef FAnimNode_BlendListBase FInternalNodeType;
+};
+
+// Exposes operations to be performed on anim state machine node contexts
+UCLASS(Experimental)
+class ANIMGRAPHRUNTIME_API UBlendListBaseLibrary : public UBlueprintFunctionLibrary
+{
+	GENERATED_BODY()
+
+public:
+
+	/** Get a blend list base context from an anim node context. */
+	UFUNCTION(BlueprintCallable, Category = "Blend List Base", meta = (BlueprintThreadSafe, ExpandEnumAsExecs = "Result"))
+	static FBlendListBaseReference ConvertToBlendListBase(const FAnimNodeReference& Node, EAnimNodeReferenceConversionResult& Result);
+
+	/** Reset target blend list node to that the next blend is executed from a blank state */
+	UFUNCTION(BlueprintCallable, Category = "Blend List Base", meta = (BlueprintThreadSafe, DisplayName = ""))
+	static void ResetNode(const FBlendListBaseReference& BlendListBase);
+
+
+};
