Remove empty display name from UBlendListBaseLibrary

#rnx

diff --git a/Engine/Source/Runtime/AnimGraphRuntime/Public/BlendListBaseLibrary.h b/Engine/Source/Runtime/AnimGraphRuntime/Public/BlendListBaseLibrary.h
--- a/Engine/Source/Runtime/AnimGraphRuntime/Public/BlendListBaseLibrary.h
+++ b/Engine/Source/Runtime/AnimGraphRuntime/Public/BlendListBaseLibrary.h
@@ -30,7 +30,7 @@ public:
 	static FBlendListBaseReference ConvertToBlendListBase(const FAnimNodeReference& Node, EAnimNodeReferenceConversionResult& Result);
 
 	/** Reset target blend list node to that the next blend is executed from a blank state */
-	UFUNCTION(BlueprintCallable, Category = "Blend List Base", meta = (BlueprintThreadSafe, DisplayName = ""))
+	UFUNCTION(BlueprintCallable, Category = "Blend List Base", meta = (BlueprintThreadSafe))
 	static void ResetNode(const FBlendListBaseReference& BlendListBase);
 
 
