Allow PlayMontageCallbackProxy to be derived from, to support plugins that want to provide customized versions of playing montages without duplicating existing code

diff --git a/Engine/Source/Runtime/AnimGraphRuntime/Private/PlayMontageCallbackProxy.cpp b/Engine/Source/Runtime/AnimGraphRuntime/Private/PlayMontageCallbackProxy.cpp
--- a/Engine/Source/Runtime/AnimGraphRuntime/Private/PlayMontageCallbackProxy.cpp
+++ b/Engine/Source/Runtime/AnimGraphRuntime/Private/PlayMontageCallbackProxy.cpp
@@ -30,7 +30,7 @@ UPlayMontageCallbackProxy* UPlayMontageCallbackProxy::CreateProxyObjectForPlayMo
 }
 
 
-void UPlayMontageCallbackProxy::PlayMontage(class USkeletalMeshComponent* InSkeletalMeshComponent, 
+bool UPlayMontageCallbackProxy::PlayMontage(class USkeletalMeshComponent* InSkeletalMeshComponent, 
 	class UAnimMontage* MontageToPlay, 
 	float PlayRate, 
 	float StartingPosition, 
@@ -73,6 +73,8 @@ void UPlayMontageCallbackProxy::PlayMontage(class USkeletalMeshComponent* InSkel
 	{
 		OnInterrupted.Broadcast(NAME_None);
 	}
+
+	return bPlayedSuccessfully;
 }
 
 bool UPlayMontageCallbackProxy::IsNotifyValid(FName NotifyName, const FBranchingPointNotifyPayload& BranchingPointNotifyPayload) const
diff --git a/Engine/Source/Runtime/AnimGraphRuntime/Public/PlayMontageCallbackProxy.h b/Engine/Source/Runtime/AnimGraphRuntime/Public/PlayMontageCallbackProxy.h
--- a/Engine/Source/Runtime/AnimGraphRuntime/Public/PlayMontageCallbackProxy.h
+++ b/Engine/Source/Runtime/AnimGraphRuntime/Public/PlayMontageCallbackProxy.h
@@ -10,8 +10,8 @@
 
 DECLARE_DYNAMIC_MULTICAST_DELEGATE_OneParam(FOnMontagePlayDelegate, FName, NotifyName);
 
-UCLASS(MinimalAPI)
-class UPlayMontageCallbackProxy : public UObject
+UCLASS()
+class ANIMGRAPHRUNTIME_API UPlayMontageCallbackProxy : public UObject
 {
 	GENERATED_UCLASS_BODY()
 
@@ -71,7 +71,9 @@ private:
 	FOnMontageBlendingOutStarted BlendingOutDelegate;
 	FOnMontageEnded MontageEndedDelegate;
 
-	void PlayMontage(
+protected:
+	// Attempts to play a montage with the specified settings. Returns whether it started or not.
+	bool PlayMontage(
 		class USkeletalMeshComponent* InSkeletalMeshComponent,
 		class UAnimMontage* MontageToPlay,
 		float PlayRate = 1.f,
