[PoseDriver]
- Added EditConditions for OnlyDriveBones, bOnlyDriveSelectedBones and Median properties.
- Rearranged Properties

diff --git a/Engine/Source/Runtime/AnimGraphRuntime/Private/AnimNodes/AnimNode_PoseDriver.cpp b/Engine/Source/Runtime/AnimGraphRuntime/Private/AnimNodes/AnimNode_PoseDriver.cpp
--- a/Engine/Source/Runtime/AnimGraphRuntime/Private/AnimNodes/AnimNode_PoseDriver.cpp
+++ b/Engine/Source/Runtime/AnimGraphRuntime/Private/AnimNodes/AnimNode_PoseDriver.cpp
@@ -11,9 +11,9 @@
 
 
 FAnimNode_PoseDriver::FAnimNode_PoseDriver()
-	: DriveSource(EPoseDriverSource::Rotation)
+	: bOnlyDriveSelectedBones(false)
+	, DriveSource(EPoseDriverSource::Rotation)
 	, DriveOutput(EPoseDriverOutput::DrivePoses)
-	, bOnlyDriveSelectedBones(false)
 	, LODThreshold(INDEX_NONE)
 {
 	RBFParams.DistanceMethod = ERBFDistanceMethod::SwingAngle;
diff --git a/Engine/Source/Runtime/AnimGraphRuntime/Public/AnimNodes/AnimNode_PoseDriver.h b/Engine/Source/Runtime/AnimGraphRuntime/Public/AnimNodes/AnimNode_PoseDriver.h
--- a/Engine/Source/Runtime/AnimGraphRuntime/Public/AnimNodes/AnimNode_PoseDriver.h
+++ b/Engine/Source/Runtime/AnimGraphRuntime/Public/AnimNodes/AnimNode_PoseDriver.h
@@ -138,33 +138,16 @@ struct ANIMGRAPHRUNTIME_API FAnimNode_PoseDriver : public FAnimNode_PoseHandler
 	GENERATED_BODY()
 
 	/** Bones to use for driving parameters based on their transform */
-	UPROPERTY(EditAnywhere, EditFixedSize, BlueprintReadWrite, Category = Links)
+	UPROPERTY(EditAnywhere, EditFixedSize, BlueprintReadWrite, Category = PoseDriver)
 	FPoseLink SourcePose;
 
 	/** Bone to use for driving parameters based on its orientation */
 	UPROPERTY(EditAnywhere, Category = PoseDriver)
 	TArray<FBoneReference> SourceBones;
 
-	/** If bFilterDrivenBones is specified, only these bones will be modified by this node */
-	UPROPERTY(EditAnywhere, Category = PoseDriver, meta = (EditCondition = "bOnlyDriveSelectedBones"))
-	TArray<FBoneReference> OnlyDriveBones;
-
-	/** Targets used to compare with current pose and drive morphs/poses */
-	UPROPERTY(EditAnywhere, Category = PoseDriver)
-	TArray<FPoseDriverTarget> PoseTargets;
-
-	/** Last set of output weights from RBF solve */
-	TArray<FRBFOutputWeight> OutputWeights;
-
-	/** Input source bone TM, used for debug drawing */
-	TArray<FTransform> SourceBoneTMs;
-
-	/** If bFilterDrivenBones, this array lists bones that we should filter out (ie have a track in the PoseAsset, but are not listed in OnlyDriveBones */
-	TArray<FCompactPoseBoneIndex> BonesToFilter;
-
-	/** 
+	/**
 	 *	Optional other bone space to use when reading SourceBone transform.
-	 *	If not specified, we just use local space of SourceBone (ie relative to parent bone) 
+	 *	If not specified, the local space of SourceBone will be used. (ie relative to parent bone)
 	 */
 	UPROPERTY(EditAnywhere, Category = PoseDriver)
 	FBoneReference EvalSpaceBone;
@@ -172,17 +155,46 @@ struct ANIMGRAPHRUNTIME_API FAnimNode_PoseDriver : public FAnimNode_PoseHandler
 	/**
 	 *	Evaluate SourceBone transform relative from its Reference Pose.
 	 *  This is recommended when using Swing and Twist Angle as Distance Method, since the twist will be computed from RefPose.
-	 * 
-	 *	If not specified, we just use local space of SourceBone (ie relative to parent bone)
+	 *
+	 *	If not specified, the local space of SourceBone will be used. (ie relative to parent bone)
 	 *  This mode won't work in conjunction with EvalSpaceBone;
 	 */
 	UPROPERTY(EditAnywhere, Category = PoseDriver)
 	bool bEvalFromRefPose = false;
 
+	/** If bOnlyDriveSelectedBones is specified, only the OnlyDriveBones will be modified by this node. */
+	UPROPERTY(EditAnywhere, Category = PoseDriver, meta = (EditCondition = "DriveOutput == EPoseDriverOutput::DrivePoses"))
+	uint8 bOnlyDriveSelectedBones : 1;
+
+	/** List of bones that will modified by this node. If no list is provided, all bones bones with a track in the PoseAsset will be modified */
+	UPROPERTY(EditAnywhere, Category = PoseDriver, meta = (EditCondition = "bOnlyDriveSelectedBones && DriveOutput == EPoseDriverOutput::DrivePoses"))
+	TArray<FBoneReference> OnlyDriveBones;
+
+	/** Targets used to compare with current pose and drive morphs/poses */
+	UPROPERTY(EditAnywhere, Category = PoseDriver)
+	TArray<FPoseDriverTarget> PoseTargets;
+
 	/** Parameters used by RBF solver */
-	UPROPERTY(EditAnywhere, Category = PoseDriver, meta=(ShowOnlyInnerProperties))
+	UPROPERTY(EditAnywhere, Category = PoseDriver, meta = (ShowOnlyInnerProperties))
 	FRBFParams RBFParams;
 
+	/** Which part of the transform is read */
+	UPROPERTY(EditAnywhere, Category = PoseDriver)
+	EPoseDriverSource DriveSource;
+
+	/** Whether we should drive poses or curves */
+	UPROPERTY(EditAnywhere, Category = PoseDriver)
+	EPoseDriverOutput DriveOutput;
+
+	/** Last set of output weights from RBF solve */
+	TArray<FRBFOutputWeight> OutputWeights;
+
+	/** Input source bone TM, used for debug drawing */
+	TArray<FTransform> SourceBoneTMs;
+
+	/** If bOnlyDriveSelectedBones, this array lists bones that we should filter out (ie have a track in the PoseAsset, but are not listed in OnlyDriveBones */
+	TArray<FCompactPoseBoneIndex> BonesToFilter;
+
 #if WITH_EDITORONLY_DATA
 	// Deprecated
 	UPROPERTY()
@@ -196,18 +208,6 @@ struct ANIMGRAPHRUNTIME_API FAnimNode_PoseDriver : public FAnimNode_PoseHandler
 	//
 #endif
 
-	/** Which part of the transform is read */
-	UPROPERTY(EditAnywhere, Category = PoseDriver)
-	EPoseDriverSource DriveSource;
-
-	/** Whether we should drive poses or curves */
-	UPROPERTY(EditAnywhere, Category = PoseDriver)
-	EPoseDriverOutput DriveOutput;
-
-	/** If we should filter bones to be driven using the DrivenBonesFilter array */
-	UPROPERTY(EditAnywhere, Category = PoseDriver)
-	uint8 bOnlyDriveSelectedBones : 1;
-
 	/** If true, will recalculate DrivenUID values in PoseTargets array on next eval */
 	uint8 bCachedDrivenIDsAreDirty : 1;
 
diff --git a/Engine/Source/Runtime/AnimGraphRuntime/Public/RBF/RBFSolver.h b/Engine/Source/Runtime/AnimGraphRuntime/Public/RBF/RBFSolver.h
--- a/Engine/Source/Runtime/AnimGraphRuntime/Public/RBF/RBFSolver.h
+++ b/Engine/Source/Runtime/AnimGraphRuntime/Public/RBF/RBFSolver.h
@@ -241,15 +241,15 @@ struct ANIMGRAPHRUNTIME_API FRBFParams
 	ERBFNormalizeMethod NormalizeMethod;
 
 	/** Rotation or position of median (used for normalization) */
-	UPROPERTY(EditAnywhere, BlueprintReadWrite, Category = RBFData)
+	UPROPERTY(EditAnywhere, BlueprintReadWrite, Category = RBFData, meta = (EditCondition = "NormalizeMethod == ERBFNormalizeMethod::NormalizeWithinMedian"))
 	FVector MedianReference;
 
 	/** Minimum distance used for median */
-	UPROPERTY(EditAnywhere, BlueprintReadWrite, Category = RBFData, meta = (UIMin = "0", UIMax = "90"))
+	UPROPERTY(EditAnywhere, BlueprintReadWrite, Category = RBFData, meta = (UIMin = "0", UIMax = "90", EditCondition = "NormalizeMethod == ERBFNormalizeMethod::NormalizeWithinMedian"))
 	float MedianMin;
 
 	/** Maximum distance used for median */
-	UPROPERTY(EditAnywhere, BlueprintReadWrite, Category = RBFData, meta = (UIMin = "0", UIMax = "90"))
+	UPROPERTY(EditAnywhere, BlueprintReadWrite, Category = RBFData, meta = (UIMin = "0", UIMax = "90", EditCondition = "NormalizeMethod == ERBFNormalizeMethod::NormalizeWithinMedian"))
 	float MedianMax;
 
 	FRBFParams();
