[PoseDriver]
- Removed bOnlyDriveSelectedBones bool. Now checks if OnlyDriveBones list has items.
- Fixed issue with OnlyDriveBones with None entry.
- Fixed indentations and improved comments.

diff --git a/Engine/Source/Editor/AnimGraph/Private/AnimGraphNode_PoseDriver.cpp b/Engine/Source/Editor/AnimGraph/Private/AnimGraphNode_PoseDriver.cpp
--- a/Engine/Source/Editor/AnimGraph/Private/AnimGraphNode_PoseDriver.cpp
+++ b/Engine/Source/Editor/AnimGraph/Private/AnimGraphNode_PoseDriver.cpp
@@ -443,16 +443,6 @@ EPoseDriverOutput& UAnimGraphNode_PoseDriver::GetPoseDriverOutput()
 	return Node.DriveOutput;
 }
 
-void UAnimGraphNode_PoseDriver::SetOnlyDriveSelectedBones(bool bOnlyDriveSelectedBones)
-{
-	Node.bOnlyDriveSelectedBones = bOnlyDriveSelectedBones;
-}
-
-bool UAnimGraphNode_PoseDriver::GetOnlyDriveSelectedBones()
-{
-	return Node.bOnlyDriveSelectedBones == 1;
-}
-
 void UAnimGraphNode_PoseDriver::AddNewTarget()
 {
 	FPoseDriverTarget& NewTarget = Node.PoseTargets[Node.PoseTargets.Add(FPoseDriverTarget())];
diff --git a/Engine/Source/Editor/AnimGraph/Public/AnimGraphNode_PoseDriver.h b/Engine/Source/Editor/AnimGraph/Public/AnimGraphNode_PoseDriver.h
--- a/Engine/Source/Editor/AnimGraph/Public/AnimGraphNode_PoseDriver.h
+++ b/Engine/Source/Editor/AnimGraph/Public/AnimGraphNode_PoseDriver.h
@@ -85,17 +85,6 @@ public:
 	UFUNCTION(BlueprintPure, Category=PoseDriver)
 	EPoseDriverOutput& GetPoseDriverOutput();
 	
-	/**  */
-	UFUNCTION(BlueprintCallable, Category=PoseDriver)
-	void SetOnlyDriveSelectedBones(bool bOnlyDriveSelectedBones);
-
-	/**  */
-	UFUNCTION(BlueprintPure, Category=PoseDriver)
-	bool GetOnlyDriveSelectedBones();
-
-
-	
-
 	/** Automatically modify TargetScale for each PoseTarget, based on distance to nearest neighbor */
 	void AutoSetTargetScales(float& OutMaxDistance);
 
diff --git a/Engine/Source/Runtime/AnimGraphRuntime/Private/AnimNodes/AnimNode_PoseDriver.cpp b/Engine/Source/Runtime/AnimGraphRuntime/Private/AnimNodes/AnimNode_PoseDriver.cpp
--- a/Engine/Source/Runtime/AnimGraphRuntime/Private/AnimNodes/AnimNode_PoseDriver.cpp
+++ b/Engine/Source/Runtime/AnimGraphRuntime/Private/AnimNodes/AnimNode_PoseDriver.cpp
@@ -11,8 +11,7 @@
 
 
 FAnimNode_PoseDriver::FAnimNode_PoseDriver()
-	: bOnlyDriveSelectedBones(false)
-	, DriveSource(EPoseDriverSource::Rotation)
+	: DriveSource(EPoseDriverSource::Rotation)
 	, DriveOutput(EPoseDriverOutput::DrivePoses)
 	, LODThreshold(INDEX_NONE)
 {
@@ -93,10 +92,20 @@ void FAnimNode_PoseDriver::CacheBones_AnyThread(const FAnimationCacheBonesContex
 		}
 	}
 
+	// Check if there are valid OnlyDriveBones, since there might be a None entry.
+	bHasOnlyDriveBones = false;
+	for (const FBoneReference& BoneRef : OnlyDriveBones)
+	{
+		if (!BoneRef.BoneName.IsNone())
+		{
+			bHasOnlyDriveBones = true;
+			break;
+		}
+	}
 
 	// If we are filtering for only specific bones, set blend weight to zero for unwanted bones, and remember which bones to filter
 	BonesToFilter.Reset();
-	if (bOnlyDriveSelectedBones && CurrentPoseAsset.IsValid())
+	if (bHasOnlyDriveBones && CurrentPoseAsset.IsValid())
 	{
 		// Super call above should init BoneBlendWeights to compact pose size if CurrentPoseAsset is valid
 		check(BoneBlendWeights.Num() == BoneContainer.GetBoneIndicesArray().Num());
@@ -104,7 +113,7 @@ void FAnimNode_PoseDriver::CacheBones_AnyThread(const FAnimationCacheBonesContex
 		const TArray<FName> TrackNames = CurrentPoseAsset.Get()->GetTrackNames();
 		for (const auto& TrackName : TrackNames)
 		{
-			// See if bone is in select list
+			// Check if bone in OnlyDriveBones
 			if (!IsBoneDriven(TrackName))
 			{
 				int32 MeshBoneIndex = BoneContainer.GetPoseBoneIndexForBoneName(TrackName);
@@ -145,23 +154,15 @@ float FAnimNode_PoseDriver::GetRadiusForTarget(const FRBFTarget& Target) const
 
 bool FAnimNode_PoseDriver::IsBoneDriven(FName BoneName) const
 {
-	// If not filtering, drive all the bones
-	if (!bOnlyDriveSelectedBones)
-	{
-		return true;
-	}
-
-	bool bIsDriven = false;
 	for (const FBoneReference& BoneRef : OnlyDriveBones)
 	{
 		if (BoneRef.BoneName == BoneName)
 		{
-			bIsDriven = true;
-			break;
+			return true;
 		}
 	}
 
-	return bIsDriven;
+	return false;
 }
 
 
@@ -257,118 +258,120 @@ void FAnimNode_PoseDriver::Evaluate_AnyThread(FPoseContext& Output)
 	FPoseContext SourceData(Output);
 	SourcePose.Evaluate(SourceData);
 
-		// Udpate DrivenIDs if needed
-		if (bCachedDrivenIDsAreDirty)
+	// Udpate DrivenIDs if needed
+	if (bCachedDrivenIDsAreDirty)
+	{
+		if (CurrentPoseAsset.IsValid())
 		{
-			if (CurrentPoseAsset.IsValid())
-			{
-				RebuildPoseList(Output.AnimInstanceProxy->GetRequiredBones(), CurrentPoseAsset.Get());
-			}
+			RebuildPoseList(Output.AnimInstanceProxy->GetRequiredBones(), CurrentPoseAsset.Get());
 		}
+	}
 
-		// Get the index of the source bone
-		const FBoneContainer& BoneContainer = SourceData.Pose.GetBoneContainer();
+	// Get the index of the source bone
+	const FBoneContainer& BoneContainer = SourceData.Pose.GetBoneContainer();
 
-		RBFInput.Values.Reset();
+	RBFInput.Values.Reset();
 
-		SourceBoneTMs.Reset();
-		bool bFoundAnyBone = false;
-		for (const FBoneReference& SourceBoneRef : SourceBones)
-		{
-			FTransform SourceBoneTM = FTransform::Identity;
+	SourceBoneTMs.Reset();
+	bool bFoundAnyBone = false;
+	for (const FBoneReference& SourceBoneRef : SourceBones)
+	{
+		FTransform SourceBoneTM = FTransform::Identity;
 
-			const FCompactPoseBoneIndex SourceCompactIndex = SourceBoneRef.GetCompactPoseIndex(BoneContainer);
-			if (SourceCompactIndex.GetInt() != INDEX_NONE)
+		const FCompactPoseBoneIndex SourceCompactIndex = SourceBoneRef.GetCompactPoseIndex(BoneContainer);
+		if (SourceCompactIndex.GetInt() != INDEX_NONE)
+		{
+			// If evaluating in alternative bone space, have to build component space pose
+			if (EvalSpaceBone.IsValidToEvaluate(BoneContainer))
 			{
-				// If evaluating in alternative bone space, have to build component space pose
-				if (EvalSpaceBone.IsValidToEvaluate(BoneContainer))
-				{
-					FCSPose<FCompactPose> CSPose;
-					CSPose.InitPose(SourceData.Pose);
+				FCSPose<FCompactPose> CSPose;
+				CSPose.InitPose(SourceData.Pose);
 
-					const FCompactPoseBoneIndex EvalSpaceCompactIndex = EvalSpaceBone.GetCompactPoseIndex(BoneContainer);
-					FTransform EvalSpaceCompSpace = CSPose.GetComponentSpaceTransform(EvalSpaceCompactIndex);
-					FTransform SourceBoneCompSpace = CSPose.GetComponentSpaceTransform(SourceCompactIndex);
+				const FCompactPoseBoneIndex EvalSpaceCompactIndex = EvalSpaceBone.GetCompactPoseIndex(BoneContainer);
+				FTransform EvalSpaceCompSpace = CSPose.GetComponentSpaceTransform(EvalSpaceCompactIndex);
+				FTransform SourceBoneCompSpace = CSPose.GetComponentSpaceTransform(SourceCompactIndex);
 
-					SourceBoneTM = SourceBoneCompSpace.GetRelativeTransform(EvalSpaceCompSpace);
+				SourceBoneTM = SourceBoneCompSpace.GetRelativeTransform(EvalSpaceCompSpace);
+			}
+			// If just evaluating in local space, just grab from local space pose
+			else
+			{
+				// Relative to Ref Pose
+				if (bEvalFromRefPose && SourceCompactIndex.GetInt() < BoneContainer.GetCompactPoseNumBones())
+				{
+					SourceBoneTM = SourceData.Pose[SourceCompactIndex].GetRelativeTransform(BoneContainer.GetRefPoseTransform(SourceCompactIndex));
 				}
-				// If just evaluating in local space, just grab from local space pose
 				else
 				{
-					// Relative to Ref Pose
-					if (bEvalFromRefPose && SourceCompactIndex.GetInt() < BoneContainer.GetCompactPoseNumBones())
-					{
-						SourceBoneTM = SourceData.Pose[SourceCompactIndex].GetRelativeTransform(BoneContainer.GetRefPoseTransform(SourceCompactIndex));
-					}
-					else
-					{
-						SourceBoneTM = SourceData.Pose[SourceCompactIndex];
-					}
+					SourceBoneTM = SourceData.Pose[SourceCompactIndex];
 				}
-
-				bFoundAnyBone = true;
 			}
 
+			bFoundAnyBone = true;
+		}
 
-			// Build RBFInput entry
-			if (DriveSource == EPoseDriverSource::Translation)
-			{
-				RBFInput.AddFromVector(SourceBoneTM.GetTranslation());
-			}
-			else
-			{
-				RBFInput.AddFromRotator(SourceBoneTM.Rotator());
-			}
 
-			// Record this so we can use it for drawing in edit mode
-			SourceBoneTMs.Add(SourceBoneTM);
+		// Build RBFInput entry
+		if (DriveSource == EPoseDriverSource::Translation)
+		{
+			RBFInput.AddFromVector(SourceBoneTM.GetTranslation());
 		}
-
-		// Do nothing if bone is no bones are found/all LOD-ed out
-		if (!bFoundAnyBone)
+		else
 		{
-			Output = SourceData;
-			return;
+			RBFInput.AddFromRotator(SourceBoneTM.Rotator());
 		}
 
-		RBFParams.TargetDimensions = SourceBones.Num() * 3;
+		// Record this so we can use it for drawing in edit mode
+		SourceBoneTMs.Add(SourceBoneTM);
+	}
+
+	// Do nothing if bone is no bones are found/all LOD-ed out
+	if (!bFoundAnyBone)
+	{
+		Output = SourceData;
+		return;
+	}
 
-		OutputWeights.Reset();
+	RBFParams.TargetDimensions = SourceBones.Num() * 3;
 
+	OutputWeights.Reset();
+	
+	// Use SoloTarget, no need to Solve
 #if WITH_EDITORONLY_DATA
-		if (SoloTargetIndex != INDEX_NONE && SoloTargetIndex < PoseTargets.Num())
-		{
-			OutputWeights.Add(FRBFOutputWeight(SoloTargetIndex, 1.0f));
-		}
-		else
+	if (PoseTargets.IsValidIndex(SoloTargetIndex))
+	{
+		OutputWeights.Add(FRBFOutputWeight(SoloTargetIndex, 1.0f));
+	}
+	else
 #endif
-		{
-			// Get target array as RBF types
-			GetRBFTargets(RBFTargets, &BoneContainer);
-
-			if (!SolverData.IsValid() || !FRBFSolver::IsSolverDataValid(*SolverData, RBFParams, RBFTargets))
-			{
-				SolverData = FRBFSolver::InitSolver(RBFParams, RBFTargets);
-			}
+	// Solve Weights
+	{
+		// Get target array as RBF types
+		GetRBFTargets(RBFTargets, &BoneContainer);
 
-			// Run RBF solver
-			FRBFSolver::Solve(*SolverData, RBFParams, RBFTargets, RBFInput, OutputWeights);
+		if (!SolverData.IsValid() || !FRBFSolver::IsSolverDataValid(*SolverData, RBFParams, RBFTargets))
+		{
+			SolverData = FRBFSolver::InitSolver(RBFParams, RBFTargets);
 		}
 
-		// Track if we have filled Output with valid pose
-		bool bHaveValidPose = false;
+		// Run RBF solver
+		FRBFSolver::Solve(*SolverData, RBFParams, RBFTargets, RBFInput, OutputWeights);
+	}
+
+	// Track if we have filled Output with valid pose
+	bool bHaveValidPose = false;
 
-		// Process active targets (if any)
-		if (OutputWeights.Num() > 0)
+	// Process active targets (if any)
+	if (OutputWeights.Num() > 0)
+	{
+		// Drive Poses from PoseAsset
+		if (DriveOutput == EPoseDriverOutput::DrivePoses)
 		{
-			// If we want to drive poses, and PoseAsset is assigned and compatible
 			const UPoseAsset* CachedPoseAsset = CurrentPoseAsset.Get();
-			if (DriveOutput == EPoseDriverOutput::DrivePoses &&
-				CachedPoseAsset &&
-				CachedPoseAsset->GetSkeleton() != nullptr)
+				
+			// Check if PoseAsset is assignedand and compatible
+			if (CachedPoseAsset && CachedPoseAsset->GetSkeleton() != nullptr)
 			{
-				FPoseContext CurrentPose(Output);
-
 				// clear the value before setting it. 
 				for (int32 PoseIndex = 0; PoseIndex < PoseExtractContext.PoseCurves.Num(); ++PoseIndex)
 				{
@@ -386,11 +389,13 @@ void FAnimNode_PoseDriver::Evaluate_AnyThread(FPoseContext& Output)
 					}
 				}
 
+				FPoseContext CurrentPose(Output);
 				FAnimationPoseData CurrentAnimationPoseData(CurrentPose);
 
-				if (CurrentPoseAsset.Get()->GetAnimationPose(CurrentAnimationPoseData, PoseExtractContext))
+				// Evaluate PoseAsset
+				if (CachedPoseAsset->GetAnimationPose(CurrentAnimationPoseData, PoseExtractContext))
 				{
-					// blend by weight
+					// If Additive, Set Source and OnlyDrive Bones to Zero and Accumulate Pose
 					if (CurrentPoseAsset->IsValidAdditive())
 					{
 						const FTransform AdditiveIdentity(FQuat::Identity, FVector::ZeroVector, FVector::ZeroVector);
@@ -403,7 +408,7 @@ void FAnimNode_PoseDriver::Evaluate_AnyThread(FPoseContext& Output)
 						}
 
 						// If filtering for specific bones, filter out bones using BonesToFilter array
-						if (bOnlyDriveSelectedBones)
+						if (bHasOnlyDriveBones)
 						{
 							for (FCompactPoseBoneIndex BoneIndex : BonesToFilter)
 							{
@@ -411,87 +416,91 @@ void FAnimNode_PoseDriver::Evaluate_AnyThread(FPoseContext& Output)
 							}
 						}
 
+						// Start by copying input
 						Output = SourceData;
-    
-						FAnimationPoseData BaseAnimationPoseData(Output);
-						const FAnimationPoseData AdditiveAnimationPoseData(CurrentPose);
-						FAnimationRuntime::AccumulateAdditivePose(BaseAnimationPoseData, AdditiveAnimationPoseData, 1.f, EAdditiveAnimationType::AAT_LocalSpaceBase);
+
+						FAnimationPoseData AccumulatedAnimationPoseData(Output); // out
+						FAnimationRuntime::AccumulateAdditivePose(AccumulatedAnimationPoseData, CurrentAnimationPoseData, 1.f, EAdditiveAnimationType::AAT_LocalSpaceBase);
 					}
+
+					// If Non-Additive, Blend between Source and Current filtering with OnlyDrive Bones
 					else
 					{
-						FAnimationPoseData BlendedAnimationPoseData(Output);
 						const FAnimationPoseData SourceAnimationPoseData(SourceData);
+						FAnimationPoseData BlendedAnimationPoseData(Output); // out
 						FAnimationRuntime::BlendTwoPosesTogetherPerBone(SourceAnimationPoseData, CurrentAnimationPoseData, BoneBlendWeights, BlendedAnimationPoseData);
 					}
 
 					bHaveValidPose = true;
 				}
 			}
-			// Drive curves (morphs, materials etc)
-			else if (DriveOutput == EPoseDriverOutput::DriveCurves)
-			{
-				// Start by copying input
-				Output = SourceData;
+		}
 
-				// Then set curves based on target weights
-				FBlendedCurve DrivenCurves;
-				for (const FRBFOutputWeight& Weight : OutputWeights)
+		// Drive curves (morphs, materials etc)
+		else if (DriveOutput == EPoseDriverOutput::DriveCurves)
+		{
+			// Start by copying input
+			Output = SourceData;
+
+			// Then set curves based on target weights
+			FBlendedCurve DrivenCurves;
+			for (const FRBFOutputWeight& Weight : OutputWeights)
+			{
+				FPoseDriverTarget& PoseTarget = PoseTargets[Weight.TargetIndex];
+				if (PoseTarget.DrivenName != NAME_None)
 				{
-					FPoseDriverTarget& PoseTarget = PoseTargets[Weight.TargetIndex];
-					if (PoseTarget.DrivenName != NAME_None)
-					{
-						DrivenCurves.Add(PoseTarget.DrivenName, Weight.TargetWeight);
-					}
+					DrivenCurves.Add(PoseTarget.DrivenName, Weight.TargetWeight);
 				}
+			}
 
-				// Merge driven curves into output
-				UE::Anim::FNamedValueArrayUtils::Union(Output.Curve, DrivenCurves, [](UE::Anim::FCurveElement& InOutElement, const UE::Anim::FCurveElement& InElement0, UE::Anim::ENamedValueUnionFlags InFlags)
+			// Merge driven curves into output
+			UE::Anim::FNamedValueArrayUtils::Union(Output.Curve, DrivenCurves, [](UE::Anim::FCurveElement& InOutElement, const UE::Anim::FCurveElement& InElement0, UE::Anim::ENamedValueUnionFlags InFlags)
+			{
+				if(EnumHasAnyFlags(InFlags, UE::Anim::ENamedValueUnionFlags::ValidArg1))
 				{
-					if(EnumHasAnyFlags(InFlags, UE::Anim::ENamedValueUnionFlags::ValidArg1))
-					{
-						InOutElement.Value = InElement0.Value;
-						InOutElement.Flags |= InElement0.Flags;
-					}
-				});
+					InOutElement.Value = InElement0.Value;
+					InOutElement.Flags |= InElement0.Flags;
+				}
+			});
 
-				bHaveValidPose = true;
-			}
+			bHaveValidPose = true;
 		}
+	}
 
-		// No valid pose, just pass through
-		if (!bHaveValidPose)
-		{
-			Output = SourceData;
-		}
+	// No valid pose, just pass through
+	if (!bHaveValidPose)
+	{
+		Output = SourceData;
+	}
 
 #if WITH_EDITORONLY_DATA
-		else if (!bSoloDrivenOnly && SoloTargetIndex != INDEX_NONE && SoloTargetIndex < PoseTargets.Num())
+	else if (!bSoloDrivenOnly && PoseTargets.IsValidIndex(SoloTargetIndex))
+	{
+		SourceBoneTMs.Reset();
+		const FPoseDriverTarget& PoseTarget = PoseTargets[SoloTargetIndex];
+
+		for (int32 SourceIdx = 0; SourceIdx < SourceBones.Num(); SourceIdx++)
 		{
-			SourceBoneTMs.Reset();
-			const FPoseDriverTarget& PoseTarget = PoseTargets[SoloTargetIndex];
+			const FBoneReference& SourceBoneRef = SourceBones[SourceIdx];
+			const FCompactPoseBoneIndex SourceCompactIndex = SourceBoneRef.GetCompactPoseIndex(BoneContainer);
 
-			for (int32 SourceIdx = 0; SourceIdx < SourceBones.Num(); SourceIdx++)
+			if (PoseTarget.BoneTransforms.IsValidIndex(SourceIdx) && SourceCompactIndex.GetInt() != INDEX_NONE)
 			{
-				const FBoneReference& SourceBoneRef = SourceBones[SourceIdx];
-				const FCompactPoseBoneIndex SourceCompactIndex = SourceBoneRef.GetCompactPoseIndex(BoneContainer);
+				FTransform& TargetTransform = Output.Pose[SourceCompactIndex];
+				const FPoseDriverTransform& SourceTransform = PoseTarget.BoneTransforms[SourceIdx];
 
-				if (PoseTarget.BoneTransforms.IsValidIndex(SourceIdx) && SourceCompactIndex.GetInt() != INDEX_NONE)
+				if (DriveSource == EPoseDriverSource::Translation)
 				{
-					FTransform& TargetTransform = Output.Pose[SourceCompactIndex];
-					const FPoseDriverTransform& SourceTransform = PoseTarget.BoneTransforms[SourceIdx];
-
-					if (DriveSource == EPoseDriverSource::Translation)
-					{
-						TargetTransform.SetTranslation(SourceTransform.TargetTranslation);
-					}
-					else
-					{
-						TargetTransform.SetRotation(SourceTransform.TargetRotation.Quaternion());
-					}
-					SourceBoneTMs.Add(TargetTransform);
+					TargetTransform.SetTranslation(SourceTransform.TargetTranslation);
 				}
+				else
+				{
+					TargetTransform.SetRotation(SourceTransform.TargetRotation.Quaternion());
+				}
+				SourceBoneTMs.Add(TargetTransform);
 			}
 		}
-#endif
 	}
+#endif
+}
 
diff --git a/Engine/Source/Runtime/AnimGraphRuntime/Public/AnimNodes/AnimNode_PoseDriver.h b/Engine/Source/Runtime/AnimGraphRuntime/Public/AnimNodes/AnimNode_PoseDriver.h
--- a/Engine/Source/Runtime/AnimGraphRuntime/Public/AnimNodes/AnimNode_PoseDriver.h
+++ b/Engine/Source/Runtime/AnimGraphRuntime/Public/AnimNodes/AnimNode_PoseDriver.h
@@ -162,12 +162,8 @@ struct ANIMGRAPHRUNTIME_API FAnimNode_PoseDriver : public FAnimNode_PoseHandler
 	UPROPERTY(EditAnywhere, Category = PoseDriver)
 	bool bEvalFromRefPose = false;
 
-	/** If bOnlyDriveSelectedBones is specified, only the OnlyDriveBones will be modified by this node. */
-	UPROPERTY(EditAnywhere, Category = PoseDriver, meta = (EditCondition = "DriveOutput == EPoseDriverOutput::DrivePoses"))
-	uint8 bOnlyDriveSelectedBones : 1;
-
 	/** List of bones that will modified by this node. If no list is provided, all bones bones with a track in the PoseAsset will be modified */
-	UPROPERTY(EditAnywhere, Category = PoseDriver, meta = (EditCondition = "bOnlyDriveSelectedBones && DriveOutput == EPoseDriverOutput::DrivePoses"))
+	UPROPERTY(EditAnywhere, Category = PoseDriver, meta = (EditCondition = "DriveOutput == EPoseDriverOutput::DrivePoses"))
 	TArray<FBoneReference> OnlyDriveBones;
 
 	/** Targets used to compare with current pose and drive morphs/poses */
@@ -191,8 +187,11 @@ struct ANIMGRAPHRUNTIME_API FAnimNode_PoseDriver : public FAnimNode_PoseHandler
 
 	/** Input source bone TM, used for debug drawing */
 	TArray<FTransform> SourceBoneTMs;
+	
+	/** Checks if there are valid bones in OnlyDriveBones, since some entries can be None */
+	bool bHasOnlyDriveBones = false;
 
-	/** If bOnlyDriveSelectedBones, this array lists bones that we should filter out (ie have a track in the PoseAsset, but are not listed in OnlyDriveBones */
+	/** This array lists bones that we should filter out (ie have a track in the PoseAsset, but are not listed in OnlyDriveBones) */
 	TArray<FCompactPoseBoneIndex> BonesToFilter;
 
 #if WITH_EDITORONLY_DATA
