Rotate root bone can opt-in to rotate the root motion custom attribute

diff --git a/Engine/Source/Runtime/AnimGraphRuntime/Private/AnimNodes/AnimNode_RotateRootBone.cpp b/Engine/Source/Runtime/AnimGraphRuntime/Private/AnimNodes/AnimNode_RotateRootBone.cpp
--- a/Engine/Source/Runtime/AnimGraphRuntime/Private/AnimNodes/AnimNode_RotateRootBone.cpp
+++ b/Engine/Source/Runtime/AnimGraphRuntime/Private/AnimNodes/AnimNode_RotateRootBone.cpp
@@ -2,6 +2,7 @@
 
 #include "AnimNodes/AnimNode_RotateRootBone.h"
 #include "Animation/AnimTrace.h"
+#include "Animation/AnimRootMotionProvider.h"
 
 #include UE_INLINE_GENERATED_CPP_BY_NAME(AnimNode_RotateRootBone)
 
@@ -61,6 +62,21 @@ void FAnimNode_RotateRootBone::Evaluate_AnyThread(FPoseContext& Output)
 		FCompactPoseBoneIndex RootBoneIndex(0);
 		Output.Pose[RootBoneIndex].SetRotation(Output.Pose[RootBoneIndex].GetRotation() * MeshSpaceDeltaQuat);
 		Output.Pose[RootBoneIndex].NormalizeRotation();
+
+		if (bRotateRootMotionAttribute)
+		{
+			// Rotate our root motion attribute by the same rotation we apply to the root
+			const UE::Anim::IAnimRootMotionProvider* RootMotionProvider = UE::Anim::IAnimRootMotionProvider::Get();
+			if (RootMotionProvider && RootMotionProvider->HasRootMotion(Output.CustomAttributes))
+			{
+				FTransform RootMotionTransformDelta;
+				if (RootMotionProvider->ExtractRootMotion(Output.CustomAttributes, RootMotionTransformDelta))
+				{
+					RootMotionTransformDelta.SetTranslation(MeshSpaceDeltaQuat.RotateVector(RootMotionTransformDelta.GetTranslation()));
+					RootMotionProvider->OverrideRootMotion(RootMotionTransformDelta, Output.CustomAttributes);
+				}
+			}
+		}
 	}
 }
 
diff --git a/Engine/Source/Runtime/AnimGraphRuntime/Public/AnimNodes/AnimNode_RotateRootBone.h b/Engine/Source/Runtime/AnimGraphRuntime/Public/AnimNodes/AnimNode_RotateRootBone.h
--- a/Engine/Source/Runtime/AnimGraphRuntime/Public/AnimNodes/AnimNode_RotateRootBone.h
+++ b/Engine/Source/Runtime/AnimGraphRuntime/Public/AnimNodes/AnimNode_RotateRootBone.h
@@ -32,6 +32,10 @@ struct ANIMGRAPHRUNTIME_API FAnimNode_RotateRootBone : public FAnimNode_Base
 	UPROPERTY(EditAnywhere, BlueprintReadWrite, Category=Settings, meta=(PinShownByDefault))
 	FRotator MeshToComponent;
 
+	// If enabled, rotating the root bone using this node will also rotate the direction of the root motion custom attribute
+	UPROPERTY(EditAnywhere, BlueprintReadWrite, Category=Advanced)
+	bool bRotateRootMotionAttribute;
+
 	float ActualPitch;
 
 	float ActualYaw;
