Fix for build error in AnimNode_RotateRootBone

diff --git a/Engine/Source/Runtime/AnimGraphRuntime/Private/AnimNodes/AnimNode_RotateRootBone.cpp b/Engine/Source/Runtime/AnimGraphRuntime/Private/AnimNodes/AnimNode_RotateRootBone.cpp
--- a/Engine/Source/Runtime/AnimGraphRuntime/Private/AnimNodes/AnimNode_RotateRootBone.cpp
+++ b/Engine/Source/Runtime/AnimGraphRuntime/Private/AnimNodes/AnimNode_RotateRootBone.cpp
@@ -96,6 +96,7 @@ FAnimNode_RotateRootBone::FAnimNode_RotateRootBone()
 	: Pitch(0.0f)
 	, Yaw(0.0f)
 	, MeshToComponent(FRotator::ZeroRotator)
+	, bRotateRootMotionAttribute(false)
 	, ActualPitch(0.f)
 	, ActualYaw(0.f)
 {
