RBAN - temp workaround and logging for OOB access of BodyAnimData

diff --git a/Engine/Source/Runtime/AnimGraphRuntime/Private/BoneControllers/AnimNode_RigidBody.cpp b/Engine/Source/Runtime/AnimGraphRuntime/Private/BoneControllers/AnimNode_RigidBody.cpp
--- a/Engine/Source/Runtime/AnimGraphRuntime/Private/BoneControllers/AnimNode_RigidBody.cpp
+++ b/Engine/Source/Runtime/AnimGraphRuntime/Private/BoneControllers/AnimNode_RigidBody.cpp
@@ -93,6 +93,11 @@ FAutoConsoleVariableRef CVarRigidBodyNodeDeferredSimulationDefault(
 bool bRBAN_DebugDraw = false;
 FAutoConsoleVariableRef CVarRigidBodyNodeDebugDraw(TEXT("p.RigidBodyNode.DebugDraw"), bRBAN_DebugDraw, TEXT("Whether to debug draw the rigid body simulation state. Requires p.Chaos.DebugDraw.Enabled 1 to function as well."), ECVF_Default);
 
+// Temporary to avoid out of bounds access issue
+bool bRBAN_InitializeBoneReferencesRangeCheckEnabled = true;
+FAutoConsoleVariableRef CVarRigidBodyNodeInitializeBoneReferencesRangeCheckEnabled(TEXT("p.RigidBodyNode.InitializeBoneReferencesRangeCheckEnabled"), bRBAN_InitializeBoneReferencesRangeCheckEnabled, TEXT(""), ECVF_Default);
+
+
 // Array of priorities that can be indexed into with CVars, since task priorities cannot be set from scalability .ini
 static UE::Tasks::ETaskPriority GRigidBodyNodeTaskPriorities[] =
 {
@@ -1805,6 +1810,19 @@ void FAnimNode_RigidBody::InitializeBoneReferences(const FBoneContainer& Require
 
 		if (BodyIndex != INDEX_NONE)
 		{
+			// Avoid and track down issues with out-of-bounds access of BodyAnimData
+			if (bRBAN_InitializeBoneReferencesRangeCheckEnabled)
+			{
+				if (!ensure(BodyAnimData.IsValidIndex(BodyIndex)))
+				{
+					UE_LOG(LogRBAN, Warning, TEXT("FAnimNode_RigidBody::InitializeBoneReferences: BodyIndex out of range. BodyIndex=%d/%d, SkeletonBoneIndex=%d/%d, CompactPoseBoneIndex=%d, RequiredBoneIndex=%d"),
+						BodyIndex, BodyAnimData.Num(), SkeletonBoneIndex, SkeletonBoneIndexToBodyIndex.Num(), CompactPoseBoneIndex, Index);
+
+					bHasInvalidBoneReference = true;
+					break;
+				}
+			}
+
 			//If we have a body we need to save it for later
 			FOutputBoneData* OutputData = new (OutputBoneData) FOutputBoneData();
 			OutputData->BodyIndex = BodyIndex;
@@ -1852,8 +1870,9 @@ void FAnimNode_RigidBody::InitializeBoneReferences(const FBoneContainer& Require
 	if (bHasInvalidBoneReference)
 	{
 		// If a bone was missing, let us know which asset it happened on, and clear our bone container to make the bad asset visible.
-		ensureMsgf(false, TEXT("FAnimNode_RigidBody::InitializeBoneReferences: The Skeleton %s, is missing bones that SkeletalMesh %s needs. Skeleton might need to be resaved."),
+		UE_LOG(LogRBAN, Warning, TEXT("FAnimNode_RigidBody::InitializeBoneReferences: The Skeleton %s, is missing bones that SkeletalMesh %s needs. Skeleton might need to be resaved."),
 			*GetNameSafe(RequiredBones.GetSkeletonAsset()), *GetNameSafe(RequiredBones.GetSkeletalMeshAsset()));
+		ensure(false);
 		OutputBoneData.Empty();
 	}
 	else
