Mark FSimSpaceSettings as POD, it has a constructor but can be copied with memcpy, this will result in STRUCT_IsPlainOldData being set, which will speed up working with UserDefinedStructs that aggregate this structure. This structure has been embedded heavily in some UserDefinedStruct (ex. RBANParamsStruct). Drawback to this change is that if we add NON-Pod data to this structure we will memcpy it.
#rnx

diff --git a/Engine/Source/Runtime/AnimGraphRuntime/Private/BoneControllers/AnimNode_RigidBody.cpp b/Engine/Source/Runtime/AnimGraphRuntime/Private/BoneControllers/AnimNode_RigidBody.cpp
--- a/Engine/Source/Runtime/AnimGraphRuntime/Private/BoneControllers/AnimNode_RigidBody.cpp
+++ b/Engine/Source/Runtime/AnimGraphRuntime/Private/BoneControllers/AnimNode_RigidBody.cpp
@@ -116,6 +116,11 @@ FAutoConsoleVariableRef CVarRigidBodyNodeSimulationTaskPriority(
 	ECVF_Default
 );
 
+// This is to validate our declaration of TIsPODType in the header, which
+// was done to ensure that STRUCT_IsPlainOldData is set, which allows scripts
+// and reflection based clients to copy via memcpy:
+static_assert(std::is_trivially_copyable<FSimSpaceSettings>::value);
+
 FSimSpaceSettings::FSimSpaceSettings()
 	: WorldAlpha(0)
 	, VelocityScaleZ(1)
diff --git a/Engine/Source/Runtime/AnimGraphRuntime/Public/BoneControllers/AnimNode_RigidBody.h b/Engine/Source/Runtime/AnimGraphRuntime/Public/BoneControllers/AnimNode_RigidBody.h
--- a/Engine/Source/Runtime/AnimGraphRuntime/Public/BoneControllers/AnimNode_RigidBody.h
+++ b/Engine/Source/Runtime/AnimGraphRuntime/Public/BoneControllers/AnimNode_RigidBody.h
@@ -47,6 +47,8 @@ enum class ESimulationTiming : uint8
  * fraction of the world space motion onto the bodies which allows Bone-Space and Component-Space simulations to react to world-space 
  * movement in a controllable way.
  */
+template <> struct TIsPODType<FSimSpaceSettings> { enum { Value = true }; };
+
 USTRUCT(BlueprintType)
 struct ANIMGRAPHRUNTIME_API FSimSpaceSettings
 {
