Added a deprecation warning for AnimNode_StateResult.h located in AnimGraphRuntime module since the file got moved to Engine/Animation in CL 25613502.

diff --git a/Engine/Source/Runtime/AnimGraphRuntime/Public/AnimNodes/AnimNode_StateResult.h b/Engine/Source/Runtime/AnimGraphRuntime/Public/AnimNodes/AnimNode_StateResult.h
new file mode 100644
--- /dev/null
+++ b/Engine/Source/Runtime/AnimGraphRuntime/Public/AnimNodes/AnimNode_StateResult.h
@@ -0,0 +1,13 @@
+﻿// Copyright Epic Games, Inc. All Rights Reserved.
+
+#pragma once
+
+// HEADER_UNIT_SKIP - Deprecated
+
+#ifdef _MSC_VER
+#pragma message(__FILE__"(8): warning: Including AnimGraphRuntime/AnimNode_StateResult.h is deprecated. The file location changed. Use AnimNode_StateResult.h from the Engine module (Engine/Animation/AnimNode_StateResult.h) instead of the AnimGraphRuntime module.")
+#else
+#pragma message("Including AnimGraphRuntime/AnimNode_StateResult.h is deprecated. The file location changed. Use AnimNode_StateResult.h from the Engine module (Engine/Animation/AnimNode_StateResult.h) instead of the AnimGraphRuntime module.")
+#endif
+
+#include "Animation/AnimNode_StateResult.h"
\ No newline at end of file
