Added twist offset to LegIK node.
This allows animators inject knee motion either manually or with procedural controls within the graph.

diff --git a/Engine/Source/Runtime/AnimGraphRuntime/Private/BoneControllers/AnimNode_LegIK.cpp b/Engine/Source/Runtime/AnimGraphRuntime/Private/BoneControllers/AnimNode_LegIK.cpp
--- a/Engine/Source/Runtime/AnimGraphRuntime/Private/BoneControllers/AnimNode_LegIK.cpp
+++ b/Engine/Source/Runtime/AnimGraphRuntime/Private/BoneControllers/AnimNode_LegIK.cpp
@@ -105,6 +105,7 @@ void FAnimNode_LegIK::EvaluateSkeletalControl_AnyThread(FComponentSpacePoseConte
 		FAnimLegIKData& LegData = LegsData[LimbIndex];
 
 		LegData.InitializeTransforms(MyAnimInstanceProxy, Output.Pose);
+		LegData.TwistOffsetDegrees = Output.Curve.Get(LegData.LegDefPtr->TwistOffsetCurveName);
 
 		// rotate hips so foot aligns with effector.
 		const bool bOrientedLegTowardsIK = OrientLegTowardsIK(LegData);
@@ -286,6 +287,22 @@ void FIKChain::ReachTarget(const FVector& InTargetLocation, double InReachPrecis
 	}
 }
 
+void FIKChain::ApplyTwistOffset(const float InTwistOffsetDegrees)
+{
+	const FVector& HeadLoc = Links[0].Location;
+	const FVector HeadToTail = Links.Last().Location - HeadLoc;
+	const FVector RotationAxis = HeadToTail.GetSafeNormal();
+
+	// Only apply twist to non tail/head links.
+ 	for (int32 Index = 1; Index < Links.Num() - 1; ++Index)
+	{
+		FVector& LinkLoc = Links[Index].Location;
+		const FVector LinkToHead = LinkLoc - HeadLoc;
+
+		LinkLoc = HeadLoc + LinkToHead.RotateAngleAxis(InTwistOffsetDegrees, RotationAxis);
+	}
+}
+
 void FIKChain::OrientAllLinksToDirection(const FVector& InDirection)
 {
 	for (int32 Index = Links.Num() - 2; Index >= 0; Index--)
@@ -370,11 +387,17 @@ bool FAnimNode_LegIK::DoLegReachIK(FAnimLegIKData& InLegData)
 	const FVector FootFKLocation = InLegData.FKLegBoneTransforms[0].GetLocation();
 	const FVector FootIKLocation = InLegData.IKFootTransform.GetLocation();
 
-	// If we're already reaching our IK Target, we have no work to do.
-	// Unless, we are applying a limit. In that case, must run solver to apply angular limit
-	const bool FootAtGoal = FootFKLocation.Equals(FootIKLocation, ReachPrecision);
-	const bool UsingAngleLimit = InLegData.LegDefPtr->bEnableRotationLimit;
-	if (FootAtGoal && !UsingAngleLimit)
+	// There's no work to do if:
+	//	- We don't have a twist offset.
+	//	- We don't need to run the solver.
+	const bool bHasTwistOffset = !FMath::IsNearlyZero(InLegData.TwistOffsetDegrees);
+	// The solver is needed if:
+	//	- Our FK foot is not at the IK goal.
+	//	- We're applying a rotation limit.
+	const bool bFootAtGoal = FootFKLocation.Equals(FootIKLocation, ReachPrecision);
+	const bool bUsingRotationLimit = InLegData.LegDefPtr->bEnableRotationLimit;
+	const bool bNeedsSolver = !bFootAtGoal || bUsingRotationLimit;
+	if (!bNeedsSolver && !bHasTwistOffset)
 	{
 		return false;
 	}
@@ -382,8 +405,16 @@ bool FAnimNode_LegIK::DoLegReachIK(FAnimLegIKData& InLegData)
 	FIKChain& IKChain = InLegData.IKChain;
 	IKChain.InitializeFromLegData(InLegData, MyAnimInstanceProxy);
 
-	const int32 MaxIterationsOverride = CVarAnimLegIKMaxIterations.GetValueOnAnyThread() > 0 ? CVarAnimLegIKMaxIterations.GetValueOnAnyThread() : MaxIterations;
-	IKChain.ReachTarget(FootIKLocation, ReachPrecision, MaxIterationsOverride);
+	if (bNeedsSolver)
+	{
+		const int32 MaxIterationsOverride = CVarAnimLegIKMaxIterations.GetValueOnAnyThread() > 0 ? CVarAnimLegIKMaxIterations.GetValueOnAnyThread() : MaxIterations;
+		IKChain.ReachTarget(FootIKLocation, ReachPrecision, MaxIterationsOverride);
+	}
+
+	if (bHasTwistOffset)
+	{
+		IKChain.ApplyTwistOffset(InLegData.TwistOffsetDegrees);
+	}
 
 	// Update bone transforms based on IKChain
 
diff --git a/Engine/Source/Runtime/AnimGraphRuntime/Public/BoneControllers/AnimNode_LegIK.h b/Engine/Source/Runtime/AnimGraphRuntime/Public/BoneControllers/AnimNode_LegIK.h
--- a/Engine/Source/Runtime/AnimGraphRuntime/Public/BoneControllers/AnimNode_LegIK.h
+++ b/Engine/Source/Runtime/AnimGraphRuntime/Public/BoneControllers/AnimNode_LegIK.h
@@ -74,6 +74,7 @@ public:
 
 	void InitializeFromLegData(FAnimLegIKData& InLegData, FAnimInstanceProxy* InAnimInstanceProxy);
 	void ReachTarget(const FVector& InTargetLocation, double InReachPrecision, int32 InMaxIterations);
+	void ApplyTwistOffset(const float InTwistOffsetDegrees);
 
 	double GetMaximumReach() const
 	{
@@ -129,6 +130,12 @@ struct FAnimLegIKDefinition
 	UPROPERTY(EditAnywhere, Category = "Settings")
 	bool bEnableKneeTwistCorrection;
 
+	/** Name of the curve to use as the twist offset angle(in degrees).
+	* This is useful for injecting knee motion, while keeping the IK chain's goal/hand and root/hip locked in place. 
+	* Reasonable values are usually between -+15 degrees, although this is depends on how far in/out the knee is in the original pose. */
+	UPROPERTY(EditAnywhere, Category = "Settings")
+	FName TwistOffsetCurveName;
+
 	FAnimLegIKDefinition()
 		: NumBonesInLimb(2)
 		, MinRotationAngle(15.f)
@@ -136,6 +143,7 @@ struct FAnimLegIKDefinition
 		, HingeRotationAxis(EAxis::None)
 		, bEnableRotationLimit(false)
 		, bEnableKneeTwistCorrection(true)
+		, TwistOffsetCurveName(NAME_None)
 	{}
 };
 
@@ -149,6 +157,7 @@ public:
 	FTransform IKFootTransform;
 	FAnimLegIKDefinition* LegDefPtr;
 	FCompactPoseBoneIndex IKFootBoneIndex;
+	float TwistOffsetDegrees;
 	int32 NumBones;
 	TArray<FCompactPoseBoneIndex> FKLegBoneIndices;
 	TArray<FTransform> FKLegBoneTransforms;
@@ -162,6 +171,7 @@ public:
 		: IKFootTransform(FTransform::Identity)
 		, LegDefPtr(nullptr)
 		, IKFootBoneIndex(INDEX_NONE)
+		, TwistOffsetDegrees(0.0f)
 		, NumBones(INDEX_NONE)
 	{}
 };
