Build fix - make some functions inline so they can be used elsewhere

diff --git a/Engine/Source/Runtime/AnimGraphRuntime/Private/BoneControllers/AnimNode_RigidBody.cpp b/Engine/Source/Runtime/AnimGraphRuntime/Private/BoneControllers/AnimNode_RigidBody.cpp
--- a/Engine/Source/Runtime/AnimGraphRuntime/Private/BoneControllers/AnimNode_RigidBody.cpp
+++ b/Engine/Source/Runtime/AnimGraphRuntime/Private/BoneControllers/AnimNode_RigidBody.cpp
@@ -259,56 +259,6 @@ void FAnimNode_RigidBody::Initialize_AnyThread(const FAnimationInitializeContext
 #endif
 }
 
-FTransform SpaceToWorldTransform(ESimulationSpace Space, const FTransform& ComponentToWorld, const FTransform& BaseBoneTM)
-{
-	switch (Space)
-	{
-	case ESimulationSpace::ComponentSpace: 
-		return ComponentToWorld;
-	case ESimulationSpace::WorldSpace: 
-		return FTransform::Identity;
-	case ESimulationSpace::BaseBoneSpace:
-		return BaseBoneTM * ComponentToWorld;
-	default:
-		return FTransform::Identity;
-	}
-}
-
-FVector WorldVectorToSpaceNoScale(ESimulationSpace Space, const FVector& WorldDir, const FTransform& ComponentToWorld, const FTransform& BaseBoneTM)
-{
-	switch(Space)
-	{
-		case ESimulationSpace::ComponentSpace: return ComponentToWorld.InverseTransformVectorNoScale(WorldDir);
-		case ESimulationSpace::WorldSpace: return WorldDir;
-		case ESimulationSpace::BaseBoneSpace:
-			return BaseBoneTM.InverseTransformVectorNoScale(ComponentToWorld.InverseTransformVectorNoScale(WorldDir));
-		default: return FVector::ZeroVector;
-	}
-}
-
-FVector WorldPositionToSpace(ESimulationSpace Space, const FVector& WorldPoint, const FTransform& ComponentToWorld, const FTransform& BaseBoneTM)
-{
-	switch (Space)
-	{
-		case ESimulationSpace::ComponentSpace: return ComponentToWorld.InverseTransformPosition(WorldPoint);
-		case ESimulationSpace::WorldSpace: return WorldPoint;
-		case ESimulationSpace::BaseBoneSpace:
-			return BaseBoneTM.InverseTransformPosition(ComponentToWorld.InverseTransformPosition(WorldPoint));
-		default: return FVector::ZeroVector;
-	}
-}
-
-FORCEINLINE_DEBUGGABLE FTransform ConvertCSTransformToSimSpace(ESimulationSpace SimulationSpace, const FTransform& InCSTransform, const FTransform& ComponentToWorld, const FTransform& BaseBoneTM)
-{
-	switch (SimulationSpace)
-	{
-		case ESimulationSpace::ComponentSpace: return InCSTransform;
-		case ESimulationSpace::WorldSpace:  return InCSTransform * ComponentToWorld; 
-		case ESimulationSpace::BaseBoneSpace: return InCSTransform.GetRelativeTransform(BaseBoneTM); break;
-		default: ensureMsgf(false, TEXT("Unsupported Simulation Space")); return InCSTransform;
-	}
-}
-
 void FAnimNode_RigidBody::UpdateComponentPose_AnyThread(const FAnimationUpdateContext& Context)
 {
 	DECLARE_SCOPE_HIERARCHICAL_COUNTER_ANIMNODE(UpdateComponentPose_AnyThread)
diff --git a/Engine/Source/Runtime/AnimGraphRuntime/Public/BoneControllers/AnimNode_RigidBody.h b/Engine/Source/Runtime/AnimGraphRuntime/Public/BoneControllers/AnimNode_RigidBody.h
--- a/Engine/Source/Runtime/AnimGraphRuntime/Public/BoneControllers/AnimNode_RigidBody.h
+++ b/Engine/Source/Runtime/AnimGraphRuntime/Public/BoneControllers/AnimNode_RigidBody.h
@@ -525,3 +525,52 @@ struct TStructOpsTypeTraits<FAnimNode_RigidBody> : public TStructOpsTypeTraitsBa
 };
 #endif
 
+FORCEINLINE_DEBUGGABLE FTransform SpaceToWorldTransform(
+	ESimulationSpace Space, const FTransform& ComponentToWorld, const FTransform& BaseBoneTM)
+{
+	switch (Space)
+	{
+	case ESimulationSpace::ComponentSpace: return ComponentToWorld;
+	case ESimulationSpace::WorldSpace: return FTransform::Identity;
+	case ESimulationSpace::BaseBoneSpace: return BaseBoneTM * ComponentToWorld;
+	default: return FTransform::Identity;
+	}
+}
+
+FORCEINLINE_DEBUGGABLE FVector WorldVectorToSpaceNoScale(
+	ESimulationSpace Space, const FVector& WorldDir, const FTransform& ComponentToWorld, const FTransform& BaseBoneTM)
+{
+	switch (Space)
+	{
+	case ESimulationSpace::ComponentSpace: return ComponentToWorld.InverseTransformVectorNoScale(WorldDir);
+	case ESimulationSpace::WorldSpace: return WorldDir;
+	case ESimulationSpace::BaseBoneSpace:
+		return BaseBoneTM.InverseTransformVectorNoScale(ComponentToWorld.InverseTransformVectorNoScale(WorldDir));
+	default: return FVector::ZeroVector;
+	}
+}
+
+FORCEINLINE_DEBUGGABLE FVector WorldPositionToSpace(
+	ESimulationSpace Space, const FVector& WorldPoint, const FTransform& ComponentToWorld, const FTransform& BaseBoneTM)
+{
+	switch (Space)
+	{
+	case ESimulationSpace::ComponentSpace: return ComponentToWorld.InverseTransformPosition(WorldPoint);
+	case ESimulationSpace::WorldSpace: return WorldPoint;
+	case ESimulationSpace::BaseBoneSpace:
+		return BaseBoneTM.InverseTransformPosition(ComponentToWorld.InverseTransformPosition(WorldPoint));
+	default: return FVector::ZeroVector;
+	}
+}
+
+FORCEINLINE_DEBUGGABLE FTransform ConvertCSTransformToSimSpace(
+	ESimulationSpace Space, const FTransform& InCSTransform, const FTransform& ComponentToWorld, const FTransform& BaseBoneTM)
+{
+	switch (Space)
+	{
+	case ESimulationSpace::ComponentSpace: return InCSTransform;
+	case ESimulationSpace::WorldSpace:  return InCSTransform * ComponentToWorld;
+	case ESimulationSpace::BaseBoneSpace: return InCSTransform.GetRelativeTransform(BaseBoneTM); break;
+	default: ensureMsgf(false, TEXT("Unsupported Simulation Space")); return InCSTransform;
+	}
+}
