[Engine]

diff --git a/Engine/Source/Runtime/Analytics/AnalyticsET/Public/AnalyticsProviderETEventCache.h b/Engine/Source/Runtime/Analytics/AnalyticsET/Public/AnalyticsProviderETEventCache.h
--- a/Engine/Source/Runtime/Analytics/AnalyticsET/Public/AnalyticsProviderETEventCache.h
+++ b/Engine/Source/Runtime/Analytics/AnalyticsET/Public/AnalyticsProviderETEventCache.h
@@ -22,7 +22,7 @@
  *
  * All public APIs in this class are threadsafe. Implemented via crappy critical sections for now, but they are safe. 
  */
-class ANALYTICSET_API FAnalyticsProviderETEventCache
+class FAnalyticsProviderETEventCache
 {
 public:
 	/** Default ctor.
@@ -32,70 +32,70 @@ public:
 	 *        This ensures that no payload ever gets too large. See AddToCache() and HasFlushesQueued() for details. 
 	 *        If negative, will use INI-configured value: Engine:[AnalyticsProviderETEventCache]MaximumPayloadSize, or 100KB if not configured.
 	 */
-	FAnalyticsProviderETEventCache(int32 MaximumPayloadSize = -1, int32 InPreallocatedPayloadSize = -1);
+	ANALYTICSET_API FAnalyticsProviderETEventCache(int32 MaximumPayloadSize = -1, int32 InPreallocatedPayloadSize = -1);
 
 	/** 
 	 * Adds a new event to the cache.
 	 * If the estimated payload size will increase beyond MaximumPayloadSize then a flush will be queued here. This will make HasFlushesQueued() == true.
 	 */
-	void AddToCache(FString EventName, const TArray<FAnalyticsEventAttribute>& Attributes);
-	void AddToCache(FString EventName);
+	ANALYTICSET_API void AddToCache(FString EventName, const TArray<FAnalyticsEventAttribute>& Attributes);
+	ANALYTICSET_API void AddToCache(FString EventName);
 
 	/**
 	 * Sets an array of attributes that will automatically be appended to any event that is sent.
 	 * Logical effect is like adding them to all events before calling RecordEvent.
 	 * Practically, it is implemented much more efficiently from a storage and allocation perspective.
 	 */
-	void SetDefaultAttributes(TArray<FAnalyticsEventAttribute>&& DefaultAttributes);
+	ANALYTICSET_API void SetDefaultAttributes(TArray<FAnalyticsEventAttribute>&& DefaultAttributes);
 
 	/**
 	 * @return the current array of default attributes.
 	 */
-	TArray<FAnalyticsEventAttribute> GetDefaultAttributes() const;
+	ANALYTICSET_API TArray<FAnalyticsEventAttribute> GetDefaultAttributes() const;
 
 	/**
 	 * @return the number of default attributes are currently being applied.
 	 */
-	int32 GetDefaultAttributeCount() const;
+	ANALYTICSET_API int32 GetDefaultAttributeCount() const;
 
 	/**
 	 * Range checking is not done, similar to TArray. Use GetDefaultAttributeCount() first!
 	 * @return one attribute of the default attributes so we don't have to copy the entire attribute array.
 	 */
-	FAnalyticsEventAttribute GetDefaultAttribute(int32 AttributeIndex) const;
+	ANALYTICSET_API FAnalyticsEventAttribute GetDefaultAttribute(int32 AttributeIndex) const;
 
 	/** Flushes the cache as a string. This method is inefficient because we build up the array directly as UTF8. If nothing is cached, returns an empty string. */
 	UE_DEPRECATED(4.25, "This method has been deprecated, use FlushCacheUTF8() instead.")
-	FString FlushCache(SIZE_T* OutEventCount = nullptr);
+	ANALYTICSET_API FString FlushCache(SIZE_T* OutEventCount = nullptr);
 
 	/** Flushes the cache as a UTF8 char array. Returns a uint8 because that's what IHttpRequest prefers. If nothing is cached, returns an empty array. */
-	TArray<uint8> FlushCacheUTF8();
+	ANALYTICSET_API TArray<uint8> FlushCacheUTF8();
 
 	/**
 	* Determines whether we have anything we need to flush, either a queued flush or existing events in the payload.
 	*/
-	bool CanFlush() const;
+	ANALYTICSET_API bool CanFlush() const;
 
 	/**
 	 * Lets external code know that there are payloads queued for flush.
 	 * This happens when AddCache() calls cause the payload size to exceed MaxPayloadSize. 
 	 * Calling code needs to notice this and flush the queue.
 	 */
-	bool HasFlushesQueued() const;
+	ANALYTICSET_API bool HasFlushesQueued() const;
 
 	/**
 	 * Gets the number of cached events (doesn't include any flushes that are already queued for flush). 
 	 */
-	int GetNumCachedEvents() const;
+	ANALYTICSET_API int GetNumCachedEvents() const;
 
 	/**
 	 * Sets the preallocated payload size
 	 * @param InPreallocatedCacheSize size to preallocate the payload buffer to during init and after flushing. If negative, will match the INI-configured MaximumPayloadSize.
 	 */
-	void SetPreallocatedPayloadSize(int32 InSetPreallocatedPayloadSize);
+	ANALYTICSET_API void SetPreallocatedPayloadSize(int32 InSetPreallocatedPayloadSize);
 
 	/** Gets the preallocated size of the payload buffer. */
-	int32 GetSetPreallocatedPayloadSize() const;
+	ANALYTICSET_API int32 GetSetPreallocatedPayloadSize() const;
 
 	friend class Lock;
 
@@ -111,7 +111,7 @@ public:
 	};
 
 private:
-	void QueueFlush();
+	ANALYTICSET_API void QueueFlush();
 
 	/**
 	* Analytics event entry to be cached
@@ -152,4 +152,4 @@ private:
 
 	/** Critical section for updating the CachedEvents. Mutable to allow const methods to access the list. */
 	mutable FCriticalSection CachedEventsCS;
-};
\ No newline at end of file
+};
diff --git a/Engine/Source/Runtime/Analytics/AnalyticsET/Public/AnalyticsSessionSummaryManager.h b/Engine/Source/Runtime/Analytics/AnalyticsET/Public/AnalyticsSessionSummaryManager.h
--- a/Engine/Source/Runtime/Analytics/AnalyticsET/Public/AnalyticsSessionSummaryManager.h
+++ b/Engine/Source/Runtime/Analytics/AnalyticsET/Public/AnalyticsSessionSummaryManager.h
@@ -47,14 +47,14 @@ ANALYTICSET_API const TCHAR* LexToString(EAnalyticsSessionShutdownType ShutdownT
  * with the left over is put on the principal process. The principal process is responsible to clean up the left over from its
  * previous execution(s) and possibly salvage and send reports delayed.
  */
-class ANALYTICSET_API FAnalyticsSessionSummaryManager
+class FAnalyticsSessionSummaryManager
 {
 public:
 	/** Special key handled by the manager for the analytics backend. Its value must be one of EAnalyticsSessionShutdownType enum value. */
-	static const TAnalyticsProperty<int32> ShutdownTypeCodeProperty;
+	static ANALYTICSET_API const TAnalyticsProperty<int32> ShutdownTypeCodeProperty;
 
 	/** Special key handled in CRC to decide if an abnormal shutdown crash report should be generated. */
-	static const TAnalyticsProperty<bool> IsUserLoggingOutProperty;
+	static ANALYTICSET_API const TAnalyticsProperty<bool> IsUserLoggingOutProperty;
 
 public:
 	/**
@@ -67,7 +67,7 @@ public:
 	 * @param SessionId the current session ID assigned by the analytics. @see IAnalyticsProviderET
 	 * @param SavedDir The directory where files will be persisted. Should be the same for the principal and the subsidiary processes. If unspecified, the manager uses its internal default.
 	 */
-	FAnalyticsSessionSummaryManager(const FString& ProcessName, const FString& ProcessGroupId, const FString& UserId, const FString& AppId, const FString& AppVersion, const FString& SessionId, const FString& SavedDir = TEXT(""));
+	ANALYTICSET_API FAnalyticsSessionSummaryManager(const FString& ProcessName, const FString& ProcessGroupId, const FString& UserId, const FString& AppId, const FString& AppVersion, const FString& SessionId, const FString& SavedDir = TEXT(""));
 
 	/**
 	 * Constructs a manager for a subsidiary process, usually a companion process that collects extra information on the behalf of the principal process.
@@ -76,17 +76,17 @@ public:
 	 * @param PrincipalProcessId The principal process ID to which this subsidiary process is associated.
 	 * @param SavedDir The directory where files will be persisted. Should be the same for the principal and the subsidiary processes. If unspecified, the manager uses its internal default.
 	 */
-	FAnalyticsSessionSummaryManager(const FString& ProcessName, const FString& ProcessGroupId, uint32 PrincipalProcessId, const FString& SavedDir = TEXT(""));
+	ANALYTICSET_API FAnalyticsSessionSummaryManager(const FString& ProcessName, const FString& ProcessGroupId, uint32 PrincipalProcessId, const FString& SavedDir = TEXT(""));
 
 	/**
 	 * Destructor.
 	 */
-	~FAnalyticsSessionSummaryManager();
+	ANALYTICSET_API ~FAnalyticsSessionSummaryManager();
 
 	/**
 	 * Discovers and salvages/sends/cleans up left over from previous execution(s), if any. Only the principal process can send left over sessions.
 	 */
-	void Tick();
+	ANALYTICSET_API void Tick();
 
 	/**
 	 * Creates a new property store associated to this manager process group. The manager only sends an analytics session summary when all
@@ -94,7 +94,7 @@ public:
 	 * ensure the session is sent as soon a possible. Subsequent calls to MakeStore will increment an internal counter in the store's filename.
 	 * @param InitialCapacity The amount of space to reserve in the file.
 	 */
-	TSharedPtr<IAnalyticsPropertyStore> MakeStore(uint32 InitialCapacity);
+	ANALYTICSET_API TSharedPtr<IAnalyticsPropertyStore> MakeStore(uint32 InitialCapacity);
 
 	/**
 	 * Sets a summary sender to enable the manager to send sessions. If the sender is null, the sessions are not sent. If the
@@ -103,13 +103,13 @@ public:
 	 * and a subsidiary process only sends the current sessions summary if it is the last process of the group to exit. The manager
 	 * deletes the session data while sending.
 	 */
-	void SetSender(TSharedPtr<IAnalyticsSessionSummarySender> Sender);
+	ANALYTICSET_API void SetSender(TSharedPtr<IAnalyticsSessionSummarySender> Sender);
 
 	/**
 	 * Sets a the user id used for reporting analytics. Allows for changing the application user after startup. Will set the user id on
 	 * all existing stores created from this session summary manager.
 	 */
-	void SetUserId(const FString& UserId);
+	ANALYTICSET_API void SetUserId(const FString& UserId);
 
 	/**
 	 * Shuts down and sends the session summaries that can be sent, if a sender is set. Remember to flush and discard the
@@ -118,14 +118,14 @@ public:
 	 * @see SetSender()
 	 * @see MakeStore()
 	 */
-	void Shutdown(bool bDiscard = false);
+	ANALYTICSET_API void Shutdown(bool bDiscard = false);
 
 	/**
 	 * Cleans up all expired files that match the analytics session filename pattern. The manager automaticallly clean up if
 	 * it runs, but if the analytics is off and the manager is not instantiated anymore, some dead files can be left over.
 	 * @param SavedDir The directory used to save the analytic files. If unspecified, use the internal default.
 	 */
-	static void CleanupExpiredFiles(const FString& SavedDir = TEXT(""));
+	static ANALYTICSET_API void CleanupExpiredFiles(const FString& SavedDir = TEXT(""));
 
 	/**
 	 * Returns the age at which a session is considered expired and shouldn't be sent anymore.
@@ -162,28 +162,28 @@ private:
 
 private:
 	/** Delegate constructor. */
-	FAnalyticsSessionSummaryManager(const FString& ProcessName, const FString& ProcessGroupId, uint32 InCurrentProcessId, uint32 PrincipalProcessId, const FString& UserId, const FString& AppId, const FString& AppVersion, const FString& SessionId, const FString& SessionRootDir);
+	ANALYTICSET_API FAnalyticsSessionSummaryManager(const FString& ProcessName, const FString& ProcessGroupId, uint32 InCurrentProcessId, uint32 PrincipalProcessId, const FString& UserId, const FString& AppId, const FString& AppVersion, const FString& SessionId, const FString& SessionRootDir);
 
 	/** Returns all processes group (key is process group id) for which files exists. */
-	TMap<FString, FProcessGroup> GetSessionFiles() const;
+	ANALYTICSET_API TMap<FString, FProcessGroup> GetSessionFiles() const;
 
 	/** Processes the information collected by the principal and its subsidiary processes and decides if a summmary can be aggregated and sent. */
-	void ProcessSummary(const FString& ProcessGroupId, const FProcessGroup& ProcessGroup);
+	ANALYTICSET_API void ProcessSummary(const FString& ProcessGroupId, const FProcessGroup& ProcessGroup);
 
 	/** Loads the summaries created for the application and its subsidiary processes (like Editor/CRC combo) and aggregates summary properties. */
-	bool AggregateSummaries(const FString& ProcessGroupId, const FProcessGroup& ProcessGroup, TArray<FAnalyticsEventAttribute>& OutSummaryProperties, TArray<FAnalyticsEventAttribute>& OutInternalProperties);
+	ANALYTICSET_API bool AggregateSummaries(const FString& ProcessGroupId, const FProcessGroup& ProcessGroup, TArray<FAnalyticsEventAttribute>& OutSummaryProperties, TArray<FAnalyticsEventAttribute>& OutInternalProperties);
 
 	/** Deletes the files used to record and perist the properties. */
-	bool CleanupFiles(const TArray<FPropertyFileInfo>& PropertyFiles, bool bOnSuccess);
+	ANALYTICSET_API bool CleanupFiles(const TArray<FPropertyFileInfo>& PropertyFiles, bool bOnSuccess);
 
 	/** Checks if this process is allowed to process and send orphan files. */
 	bool IsOrphanGroupsOwner() const { return bOrphanGroupOwner; }
 
 	/** Returns whether this manager instance represents the principal process of a group. */
-	bool IsPrincipalProcess() const;
+	ANALYTICSET_API bool IsPrincipalProcess() const;
 
 	/** Prune the list of property stores for invalid entries. */
-	void PruneTrackedPropertyStores();
+	ANALYTICSET_API void PruneTrackedPropertyStores();
 
 private:
 	TArray<TWeakPtr<IAnalyticsPropertyStore>> WeakPropertyStores;
diff --git a/Engine/Source/Runtime/Analytics/AnalyticsET/Public/AnalyticsSessionSummarySender.h b/Engine/Source/Runtime/Analytics/AnalyticsET/Public/AnalyticsSessionSummarySender.h
--- a/Engine/Source/Runtime/Analytics/AnalyticsET/Public/AnalyticsSessionSummarySender.h
+++ b/Engine/Source/Runtime/Analytics/AnalyticsET/Public/AnalyticsSessionSummarySender.h
@@ -14,7 +14,7 @@ struct FAnalyticsEventAttribute;
 /**
  * Sends the analytics session summary to Epic Games analytics service.
  */
-class ANALYTICSET_API FAnalyticsSessionSummarySender : public IAnalyticsSessionSummarySender
+class FAnalyticsSessionSummarySender : public IAnalyticsSessionSummarySender
 {
 public:
 	/**
@@ -22,7 +22,7 @@ public:
 	 * @param Provider The analytics provider that will be used to emits the summary event.
 	 * @param ShouldEmitFilterFunc A filter function invoked for each properties before sending. If the function is bound and returns false, the property will be filtered out. If unbound, all properties passed to SendSesionSummary() are emitted.
 	 */
-	FAnalyticsSessionSummarySender(IAnalyticsProviderET& Provider, TFunction<bool(const FAnalyticsEventAttribute&)> ShouldEmitFilterFunc = TFunction<bool(const FAnalyticsEventAttribute&)>());
+	ANALYTICSET_API FAnalyticsSessionSummarySender(IAnalyticsProviderET& Provider, TFunction<bool(const FAnalyticsEventAttribute&)> ShouldEmitFilterFunc = TFunction<bool(const FAnalyticsEventAttribute&)>());
 
 	/**
 	 * Emits the summary events for the specified session id on behalf of the specified user/app/appversion. The function filters the properties by invoking the functor specified at construction.
@@ -32,7 +32,7 @@ public:
 	 * @param SessionId The session is for which the report is emitted.
 	 * @param Properties The list of properties that makes up the summary event.
 	 */
-	virtual bool SendSessionSummary(const FString& UserId, const FString& AppId, const FString& AppVersion, const FString& SessionId, const TArray<FAnalyticsEventAttribute>& Properties) override;
+	ANALYTICSET_API virtual bool SendSessionSummary(const FString& UserId, const FString& AppId, const FString& AppVersion, const FString& SessionId, const TArray<FAnalyticsEventAttribute>& Properties) override;
 
 private:
 	IAnalyticsProviderET& AnalyticsProvider;
diff --git a/Engine/Source/Runtime/AnimGraphRuntime/Public/AnimCustomInstanceHelper.h b/Engine/Source/Runtime/AnimGraphRuntime/Public/AnimCustomInstanceHelper.h
--- a/Engine/Source/Runtime/AnimGraphRuntime/Public/AnimCustomInstanceHelper.h
+++ b/Engine/Source/Runtime/AnimGraphRuntime/Public/AnimCustomInstanceHelper.h
@@ -31,7 +31,7 @@
 #include "Engine/World.h"
 #include "SequencerAnimationSupport.h"
 
-class ANIMGRAPHRUNTIME_API FAnimCustomInstanceHelper
+class FAnimCustomInstanceHelper
 {
 public:
 	/** 
@@ -286,5 +286,5 @@ public:
 
 private:
 	/** Helper function for BindToSkeletalMeshComponent */
-	static bool ShouldCreateCustomInstancePlayer(const USkeletalMeshComponent* SkeletalMeshComponent);
-};
\ No newline at end of file
+	static ANIMGRAPHRUNTIME_API bool ShouldCreateCustomInstancePlayer(const USkeletalMeshComponent* SkeletalMeshComponent);
+};
diff --git a/Engine/Source/Runtime/AnimGraphRuntime/Public/AnimExecutionContextLibrary.h b/Engine/Source/Runtime/AnimGraphRuntime/Public/AnimExecutionContextLibrary.h
--- a/Engine/Source/Runtime/AnimGraphRuntime/Public/AnimExecutionContextLibrary.h
+++ b/Engine/Source/Runtime/AnimGraphRuntime/Public/AnimExecutionContextLibrary.h
@@ -11,8 +11,8 @@
 class UAnimInstance;
 
 // Exposes operations to be performed on anim node contexts
-UCLASS()
-class ANIMGRAPHRUNTIME_API UAnimExecutionContextLibrary : public UBlueprintFunctionLibrary
+UCLASS(MinimalAPI)
+class UAnimExecutionContextLibrary : public UBlueprintFunctionLibrary
 {
 	GENERATED_BODY()
 	
@@ -29,33 +29,33 @@ public:
 	
 	/** Get the anim instance that hosts this context */
 	UFUNCTION(BlueprintPure, Category = "Animation|Utilities", meta=(BlueprintThreadSafe))
-	static UAnimInstance* GetAnimInstance(const FAnimExecutionContext& Context);
+	static ANIMGRAPHRUNTIME_API UAnimInstance* GetAnimInstance(const FAnimExecutionContext& Context);
 
 	/** Internal compiler use only - Get a reference to an anim node by index */
 	UFUNCTION(BlueprintPure, BlueprintInternalUseOnly, meta=(BlueprintThreadSafe, DefaultToSelf = "Instance"))
-	static FAnimNodeReference GetAnimNodeReference(UAnimInstance* Instance, int32 Index);
+	static ANIMGRAPHRUNTIME_API FAnimNodeReference GetAnimNodeReference(UAnimInstance* Instance, int32 Index);
 
 	/** Convert to an initialization context */
 	UFUNCTION(BlueprintCallable, Category = "Animation|Utilities", meta=(BlueprintThreadSafe, ExpandEnumAsExecs = "Result"))
-	static FAnimInitializationContext ConvertToInitializationContext(const FAnimExecutionContext& Context, EAnimExecutionContextConversionResult& Result);
+	static ANIMGRAPHRUNTIME_API FAnimInitializationContext ConvertToInitializationContext(const FAnimExecutionContext& Context, EAnimExecutionContextConversionResult& Result);
 	
 	/** Convert to an update context */
 	UFUNCTION(BlueprintCallable, Category = "Animation|Utilities", meta=(BlueprintThreadSafe, ExpandEnumAsExecs = "Result"))
-	static FAnimUpdateContext ConvertToUpdateContext(const FAnimExecutionContext& Context, EAnimExecutionContextConversionResult& Result);
+	static ANIMGRAPHRUNTIME_API FAnimUpdateContext ConvertToUpdateContext(const FAnimExecutionContext& Context, EAnimExecutionContextConversionResult& Result);
 
 	/** Get the current delta time in seconds */
 	UFUNCTION(BlueprintPure, Category = "Animation|Utilities", meta=(BlueprintThreadSafe))
-	static float GetDeltaTime(const FAnimUpdateContext& Context);
+	static ANIMGRAPHRUNTIME_API float GetDeltaTime(const FAnimUpdateContext& Context);
 
 	/** Get the current weight of this branch of the graph */
 	UFUNCTION(BlueprintPure, Category = "Animation|Utilities", meta=(BlueprintThreadSafe))
-	static float GetCurrentWeight(const FAnimUpdateContext& Context);
+	static ANIMGRAPHRUNTIME_API float GetCurrentWeight(const FAnimUpdateContext& Context);
 	
 	/** Convert to a pose context */
 	UFUNCTION(BlueprintCallable, Category = "Animation|Utilities", meta=(BlueprintThreadSafe, ExpandEnumAsExecs = "Result"))
-	static FAnimPoseContext ConvertToPoseContext(const FAnimExecutionContext& Context, EAnimExecutionContextConversionResult& Result);
+	static ANIMGRAPHRUNTIME_API FAnimPoseContext ConvertToPoseContext(const FAnimExecutionContext& Context, EAnimExecutionContextConversionResult& Result);
 
 	/** Convert to a component space pose context */
 	UFUNCTION(BlueprintCallable, Category = "Animation|Utilities", meta=(BlueprintThreadSafe, ExpandEnumAsExecs = "Result"))
-	static FAnimComponentSpacePoseContext ConvertToComponentSpacePoseContext(const FAnimExecutionContext& Context, EAnimExecutionContextConversionResult& Result);
-};
\ No newline at end of file
+	static ANIMGRAPHRUNTIME_API FAnimComponentSpacePoseContext ConvertToComponentSpacePoseContext(const FAnimExecutionContext& Context, EAnimExecutionContextConversionResult& Result);
+};
diff --git a/Engine/Source/Runtime/AnimGraphRuntime/Public/AnimNodes/AnimNode_AimOffsetLookAt.h b/Engine/Source/Runtime/AnimGraphRuntime/Public/AnimNodes/AnimNode_AimOffsetLookAt.h
--- a/Engine/Source/Runtime/AnimGraphRuntime/Public/AnimNodes/AnimNode_AimOffsetLookAt.h
+++ b/Engine/Source/Runtime/AnimGraphRuntime/Public/AnimNodes/AnimNode_AimOffsetLookAt.h
@@ -13,7 +13,7 @@
  * Yaw and Pitch directions for a referenced aim offset given a point in the world to look at.
  */
 USTRUCT(BlueprintInternalUseOnly)
-struct ANIMGRAPHRUNTIME_API FAnimNode_AimOffsetLookAt : public FAnimNode_BlendSpacePlayer
+struct FAnimNode_AimOffsetLookAt : public FAnimNode_BlendSpacePlayer
 {
 	GENERATED_BODY()
 
@@ -72,21 +72,21 @@ struct ANIMGRAPHRUNTIME_API FAnimNode_AimOffsetLookAt : public FAnimNode_BlendSp
 	bool bIsLODEnabled;
 
 public:
-	FAnimNode_AimOffsetLookAt();
+	ANIMGRAPHRUNTIME_API FAnimNode_AimOffsetLookAt();
 
 	// FAnimNode_Base interface
-	virtual void OnInitializeAnimInstance(const FAnimInstanceProxy* InProxy, const UAnimInstance* InAnimInstance) override;
+	ANIMGRAPHRUNTIME_API virtual void OnInitializeAnimInstance(const FAnimInstanceProxy* InProxy, const UAnimInstance* InAnimInstance) override;
 	virtual bool NeedsOnInitializeAnimInstance() const override { return true; }
-	virtual void Initialize_AnyThread(const FAnimationInitializeContext& Context) override;
-	virtual void CacheBones_AnyThread(const FAnimationCacheBonesContext& Context) override;
-	virtual void UpdateAssetPlayer(const FAnimationUpdateContext& Context) override;
-	virtual void Evaluate_AnyThread(FPoseContext& Output) override;
-	virtual void GatherDebugData(FNodeDebugData& DebugData) override;
+	ANIMGRAPHRUNTIME_API virtual void Initialize_AnyThread(const FAnimationInitializeContext& Context) override;
+	ANIMGRAPHRUNTIME_API virtual void CacheBones_AnyThread(const FAnimationCacheBonesContext& Context) override;
+	ANIMGRAPHRUNTIME_API virtual void UpdateAssetPlayer(const FAnimationUpdateContext& Context) override;
+	ANIMGRAPHRUNTIME_API virtual void Evaluate_AnyThread(FPoseContext& Output) override;
+	ANIMGRAPHRUNTIME_API virtual void GatherDebugData(FNodeDebugData& DebugData) override;
 	virtual int32 GetLODThreshold() const override { return LODThreshold; }
 	// End of FAnimNode_Base interface
 
 	// FAnimNode_BlendSpacePlayer interface
-	virtual FVector GetPosition() const override;
+	ANIMGRAPHRUNTIME_API virtual FVector GetPosition() const override;
 
-	void UpdateFromLookAtTarget(FPoseContext& LocalPoseContext);
+	ANIMGRAPHRUNTIME_API void UpdateFromLookAtTarget(FPoseContext& LocalPoseContext);
 };
diff --git a/Engine/Source/Runtime/AnimGraphRuntime/Public/AnimNodes/AnimNode_ApplyAdditive.h b/Engine/Source/Runtime/AnimGraphRuntime/Public/AnimNodes/AnimNode_ApplyAdditive.h
--- a/Engine/Source/Runtime/AnimGraphRuntime/Public/AnimNodes/AnimNode_ApplyAdditive.h
+++ b/Engine/Source/Runtime/AnimGraphRuntime/Public/AnimNodes/AnimNode_ApplyAdditive.h
@@ -9,7 +9,7 @@
 #include "AnimNode_ApplyAdditive.generated.h"
 
 USTRUCT(BlueprintInternalUseOnly)
-struct ANIMGRAPHRUNTIME_API FAnimNode_ApplyAdditive : public FAnimNode_Base
+struct FAnimNode_ApplyAdditive : public FAnimNode_Base
 {
 	GENERATED_USTRUCT_BODY()
 
@@ -52,14 +52,14 @@ struct ANIMGRAPHRUNTIME_API FAnimNode_ApplyAdditive : public FAnimNode_Base
 	bool bAlphaBoolEnabled;
 
 public:	
-	FAnimNode_ApplyAdditive();
+	ANIMGRAPHRUNTIME_API FAnimNode_ApplyAdditive();
 
 	// FAnimNode_Base interface
-	virtual void Initialize_AnyThread(const FAnimationInitializeContext& Context) override;
-	virtual void CacheBones_AnyThread(const FAnimationCacheBonesContext& Context) override;
-	virtual void Update_AnyThread(const FAnimationUpdateContext& Context) override;
-	virtual void Evaluate_AnyThread(FPoseContext& Output) override;
-	virtual void GatherDebugData(FNodeDebugData& DebugData) override;
+	ANIMGRAPHRUNTIME_API virtual void Initialize_AnyThread(const FAnimationInitializeContext& Context) override;
+	ANIMGRAPHRUNTIME_API virtual void CacheBones_AnyThread(const FAnimationCacheBonesContext& Context) override;
+	ANIMGRAPHRUNTIME_API virtual void Update_AnyThread(const FAnimationUpdateContext& Context) override;
+	ANIMGRAPHRUNTIME_API virtual void Evaluate_AnyThread(FPoseContext& Output) override;
+	ANIMGRAPHRUNTIME_API virtual void GatherDebugData(FNodeDebugData& DebugData) override;
 	virtual int32 GetLODThreshold() const override { return LODThreshold; }
 	// End of FAnimNode_Base interface
 
diff --git a/Engine/Source/Runtime/AnimGraphRuntime/Public/AnimNodes/AnimNode_BlendBoneByChannel.h b/Engine/Source/Runtime/AnimGraphRuntime/Public/AnimNodes/AnimNode_BlendBoneByChannel.h
--- a/Engine/Source/Runtime/AnimGraphRuntime/Public/AnimNodes/AnimNode_BlendBoneByChannel.h
+++ b/Engine/Source/Runtime/AnimGraphRuntime/Public/AnimNodes/AnimNode_BlendBoneByChannel.h
@@ -43,7 +43,7 @@ struct FBlendBoneByChannelEntry
 };
 
 USTRUCT(BlueprintInternalUseOnly)
-struct ANIMGRAPHRUNTIME_API FAnimNode_BlendBoneByChannel : public FAnimNode_Base
+struct FAnimNode_BlendBoneByChannel : public FAnimNode_Base
 {
 	GENERATED_USTRUCT_BODY()
 
@@ -91,11 +91,11 @@ public:
 	}
 
 	// FAnimNode_Base interface
-	virtual void Initialize_AnyThread(const FAnimationInitializeContext& Context) override;
-	virtual void CacheBones_AnyThread(const FAnimationCacheBonesContext& Context) override;
-	virtual void Update_AnyThread(const FAnimationUpdateContext& Context) override;
-	virtual void Evaluate_AnyThread(FPoseContext& Output) override;
-	virtual void GatherDebugData(FNodeDebugData& DebugData) override;
+	ANIMGRAPHRUNTIME_API virtual void Initialize_AnyThread(const FAnimationInitializeContext& Context) override;
+	ANIMGRAPHRUNTIME_API virtual void CacheBones_AnyThread(const FAnimationCacheBonesContext& Context) override;
+	ANIMGRAPHRUNTIME_API virtual void Update_AnyThread(const FAnimationUpdateContext& Context) override;
+	ANIMGRAPHRUNTIME_API virtual void Evaluate_AnyThread(FPoseContext& Output) override;
+	ANIMGRAPHRUNTIME_API virtual void GatherDebugData(FNodeDebugData& DebugData) override;
 	// End of FAnimNode_Base interface
 };
 
diff --git a/Engine/Source/Runtime/AnimGraphRuntime/Public/AnimNodes/AnimNode_BlendListBase.h b/Engine/Source/Runtime/AnimGraphRuntime/Public/AnimNodes/AnimNode_BlendListBase.h
--- a/Engine/Source/Runtime/AnimGraphRuntime/Public/AnimNodes/AnimNode_BlendListBase.h
+++ b/Engine/Source/Runtime/AnimGraphRuntime/Public/AnimNodes/AnimNode_BlendListBase.h
@@ -21,7 +21,7 @@ enum class EBlendListTransitionType : uint8
 
 // Blend list node; has many children
 USTRUCT(BlueprintInternalUseOnly)
-struct ANIMGRAPHRUNTIME_API FAnimNode_BlendListBase : public FAnimNode_Base
+struct FAnimNode_BlendListBase : public FAnimNode_Base
 {
 	GENERATED_BODY()
 
@@ -74,11 +74,11 @@ public:
 	FAnimNode_BlendListBase() = default;
 
 	// FAnimNode_Base interface
-	virtual void Initialize_AnyThread(const FAnimationInitializeContext& Context) override;
-	virtual void CacheBones_AnyThread(const FAnimationCacheBonesContext& Context) override;
-	virtual void Update_AnyThread(const FAnimationUpdateContext& Context) override;
-	virtual void Evaluate_AnyThread(FPoseContext& Output) override;
-	virtual void GatherDebugData(FNodeDebugData& DebugData) override;
+	ANIMGRAPHRUNTIME_API virtual void Initialize_AnyThread(const FAnimationInitializeContext& Context) override;
+	ANIMGRAPHRUNTIME_API virtual void CacheBones_AnyThread(const FAnimationCacheBonesContext& Context) override;
+	ANIMGRAPHRUNTIME_API virtual void Update_AnyThread(const FAnimationUpdateContext& Context) override;
+	ANIMGRAPHRUNTIME_API virtual void Evaluate_AnyThread(FPoseContext& Output) override;
+	ANIMGRAPHRUNTIME_API virtual void GatherDebugData(FNodeDebugData& DebugData) override;
 	// End of FAnimNode_Base interface
 
 #if WITH_EDITOR
@@ -97,28 +97,28 @@ public:
 
 public:
 	// Get the array of blend times to apply to our input poses
-	const TArray<float>& GetBlendTimes() const;
+	ANIMGRAPHRUNTIME_API const TArray<float>& GetBlendTimes() const;
 
 	// Get the type of transition that this blend list will make
-	EBlendListTransitionType GetTransitionType() const;
+	ANIMGRAPHRUNTIME_API EBlendListTransitionType GetTransitionType() const;
 
 	// Get the blend type we will use when blending
-	EAlphaBlendOption GetBlendType() const;
+	ANIMGRAPHRUNTIME_API EAlphaBlendOption GetBlendType() const;
 	
 	/** Get whether to reinitialize the child pose when re-activated. For example, when active child changes */
-	bool GetResetChildOnActivation() const;
+	ANIMGRAPHRUNTIME_API bool GetResetChildOnActivation() const;
 
 	// Get the custom blend curve to apply when blending, if any
-	UCurveFloat* GetCustomBlendCurve() const;
+	ANIMGRAPHRUNTIME_API UCurveFloat* GetCustomBlendCurve() const;
 
 	// Get the blend profile to use when blending, if any
-	UBlendProfile* GetBlendProfile() const;
+	ANIMGRAPHRUNTIME_API UBlendProfile* GetBlendProfile() const;
 	
 protected:
 	virtual int32 GetActiveChildIndex() { return 0; }
 	virtual FString GetNodeName(FNodeDebugData& DebugData) { return DebugData.GetNodeName(this); }
 
-	void Initialize();
+	ANIMGRAPHRUNTIME_API void Initialize();
 
 	friend class UBlendListBaseLibrary;
 };
diff --git a/Engine/Source/Runtime/AnimGraphRuntime/Public/AnimNodes/AnimNode_BlendListByBool.h b/Engine/Source/Runtime/AnimGraphRuntime/Public/AnimNodes/AnimNode_BlendListByBool.h
--- a/Engine/Source/Runtime/AnimGraphRuntime/Public/AnimNodes/AnimNode_BlendListByBool.h
+++ b/Engine/Source/Runtime/AnimGraphRuntime/Public/AnimNodes/AnimNode_BlendListByBool.h
@@ -9,7 +9,7 @@
 
 // This node is effectively a 'branch', picking one of two input poses based on an input Boolean value
 USTRUCT(BlueprintInternalUseOnly)
-struct ANIMGRAPHRUNTIME_API FAnimNode_BlendListByBool : public FAnimNode_BlendListBase
+struct FAnimNode_BlendListByBool : public FAnimNode_BlendListBase
 {
 	GENERATED_BODY()
 private:
@@ -23,9 +23,9 @@ public:
 	FAnimNode_BlendListByBool() = default;
 
 	// Get which input should be connected to the output
-	bool GetActiveValue() const;
+	ANIMGRAPHRUNTIME_API bool GetActiveValue() const;
 	
 protected:
-	virtual int32 GetActiveChildIndex() override;
+	ANIMGRAPHRUNTIME_API virtual int32 GetActiveChildIndex() override;
 	virtual FString GetNodeName(FNodeDebugData& DebugData) override { return DebugData.GetNodeName(this); }
 };
diff --git a/Engine/Source/Runtime/AnimGraphRuntime/Public/AnimNodes/AnimNode_BlendListByEnum.h b/Engine/Source/Runtime/AnimGraphRuntime/Public/AnimNodes/AnimNode_BlendListByEnum.h
--- a/Engine/Source/Runtime/AnimGraphRuntime/Public/AnimNodes/AnimNode_BlendListByEnum.h
+++ b/Engine/Source/Runtime/AnimGraphRuntime/Public/AnimNodes/AnimNode_BlendListByEnum.h
@@ -9,7 +9,7 @@
 
 // Blend List by Enum, it changes based on enum input that enters
 USTRUCT(BlueprintInternalUseOnly)
-struct ANIMGRAPHRUNTIME_API FAnimNode_BlendListByEnum : public FAnimNode_BlendListBase
+struct FAnimNode_BlendListByEnum : public FAnimNode_BlendListBase
 {
 	GENERATED_BODY()
 	
@@ -29,16 +29,16 @@ public:
 
 #if WITH_EDITORONLY_DATA
 	// Set the mapping from enum value to BlendPose index. Called during compilation.
-	void SetEnumToPoseIndex(const TArray<int32>& InEnumToPoseIndex);
+	ANIMGRAPHRUNTIME_API void SetEnumToPoseIndex(const TArray<int32>& InEnumToPoseIndex);
 #endif
 	
 	// Get the mapping from enum value to BlendPose index; there will be one entry per entry in the enum; entries out of range always map to pose index 0
-	const TArray<int32>& GetEnumToPoseIndex() const;
+	ANIMGRAPHRUNTIME_API const TArray<int32>& GetEnumToPoseIndex() const;
 	
 	// Get the currently selected pose (as an enum value)
-	uint8 GetActiveEnumValue() const;
+	ANIMGRAPHRUNTIME_API uint8 GetActiveEnumValue() const;
 	
 protected:
-	virtual int32 GetActiveChildIndex() override;
+	ANIMGRAPHRUNTIME_API virtual int32 GetActiveChildIndex() override;
 	virtual FString GetNodeName(FNodeDebugData& DebugData) override { return DebugData.GetNodeName(this); }
 };
diff --git a/Engine/Source/Runtime/AnimGraphRuntime/Public/AnimNodes/AnimNode_BlendListByInt.h b/Engine/Source/Runtime/AnimGraphRuntime/Public/AnimNodes/AnimNode_BlendListByInt.h
--- a/Engine/Source/Runtime/AnimGraphRuntime/Public/AnimNodes/AnimNode_BlendListByInt.h
+++ b/Engine/Source/Runtime/AnimGraphRuntime/Public/AnimNodes/AnimNode_BlendListByInt.h
@@ -9,7 +9,7 @@
 
 // Blend list node; has many children
 USTRUCT(BlueprintInternalUseOnly)
-struct ANIMGRAPHRUNTIME_API FAnimNode_BlendListByInt : public FAnimNode_BlendListBase
+struct FAnimNode_BlendListByInt : public FAnimNode_BlendListBase
 {
 	GENERATED_BODY()
 private:
@@ -22,7 +22,7 @@ public:
 	FAnimNode_BlendListByInt() = default;
 
 	// Get the currently active child index
-	virtual int32 GetActiveChildIndex() override;
+	ANIMGRAPHRUNTIME_API virtual int32 GetActiveChildIndex() override;
 	
 protected:
 	virtual FString GetNodeName(FNodeDebugData& DebugData) override { return DebugData.GetNodeName(this); }
diff --git a/Engine/Source/Runtime/AnimGraphRuntime/Public/AnimNodes/AnimNode_BlendSpaceEvaluator.h b/Engine/Source/Runtime/AnimGraphRuntime/Public/AnimNodes/AnimNode_BlendSpaceEvaluator.h
--- a/Engine/Source/Runtime/AnimGraphRuntime/Public/AnimNodes/AnimNode_BlendSpaceEvaluator.h
+++ b/Engine/Source/Runtime/AnimGraphRuntime/Public/AnimNodes/AnimNode_BlendSpaceEvaluator.h
@@ -10,7 +10,7 @@
 // Evaluates a point in a blendspace, using a specific time input rather than advancing time internally.
 // Typically the playback position of the animation for this node will represent something other than time, like jump height.
 USTRUCT(BlueprintInternalUseOnly)
-struct ANIMGRAPHRUNTIME_API FAnimNode_BlendSpaceEvaluator : public FAnimNode_BlendSpacePlayer
+struct FAnimNode_BlendSpaceEvaluator : public FAnimNode_BlendSpacePlayer
 {
 	GENERATED_USTRUCT_BODY()
 
@@ -25,15 +25,15 @@ public:
 	bool bTeleportToNormalizedTime = true;
 
 public:	
-	FAnimNode_BlendSpaceEvaluator();
+	ANIMGRAPHRUNTIME_API FAnimNode_BlendSpaceEvaluator();
 
 	// FAnimNode_Base interface
-	virtual void UpdateAssetPlayer(const FAnimationUpdateContext& Context) override;
-	virtual void GatherDebugData(FNodeDebugData& DebugData) override;
+	ANIMGRAPHRUNTIME_API virtual void UpdateAssetPlayer(const FAnimationUpdateContext& Context) override;
+	ANIMGRAPHRUNTIME_API virtual void GatherDebugData(FNodeDebugData& DebugData) override;
 	// End of FAnimNode_Base interface
 
 	// FAnimNode_BlendSpacePlayer interface
-	virtual float GetPlayRate() const override;
+	ANIMGRAPHRUNTIME_API virtual float GetPlayRate() const override;
 	virtual bool ShouldTeleportToTime() const override { return bTeleportToNormalizedTime; }
 	virtual bool IsEvaluator() const override { return true; }
 	// End of FAnimNode_BlendSpacePlayer
diff --git a/Engine/Source/Runtime/AnimGraphRuntime/Public/AnimNodes/AnimNode_BlendSpaceGraph.h b/Engine/Source/Runtime/AnimGraphRuntime/Public/AnimNodes/AnimNode_BlendSpaceGraph.h
--- a/Engine/Source/Runtime/AnimGraphRuntime/Public/AnimNodes/AnimNode_BlendSpaceGraph.h
+++ b/Engine/Source/Runtime/AnimGraphRuntime/Public/AnimNodes/AnimNode_BlendSpaceGraph.h
@@ -8,7 +8,7 @@
 
 // Allows multiple animations to be blended between based on input parameters
 USTRUCT(BlueprintInternalUseOnly)
-struct ANIMGRAPHRUNTIME_API FAnimNode_BlendSpaceGraph : public FAnimNode_BlendSpaceGraphBase
+struct FAnimNode_BlendSpaceGraph : public FAnimNode_BlendSpaceGraphBase
 {
 	GENERATED_BODY()
 
diff --git a/Engine/Source/Runtime/AnimGraphRuntime/Public/AnimNodes/AnimNode_BlendSpaceGraphBase.h b/Engine/Source/Runtime/AnimGraphRuntime/Public/AnimNodes/AnimNode_BlendSpaceGraphBase.h
--- a/Engine/Source/Runtime/AnimGraphRuntime/Public/AnimNodes/AnimNode_BlendSpaceGraphBase.h
+++ b/Engine/Source/Runtime/AnimGraphRuntime/Public/AnimNodes/AnimNode_BlendSpaceGraphBase.h
@@ -12,7 +12,7 @@ class UBlendSpace;
 
 // Allows multiple animations to be blended between based on input parameters
 USTRUCT(BlueprintInternalUseOnly)
-struct ANIMGRAPHRUNTIME_API FAnimNode_BlendSpaceGraphBase : public FAnimNode_Base
+struct FAnimNode_BlendSpaceGraphBase : public FAnimNode_Base
 {
 	GENERATED_BODY()
 
@@ -29,11 +29,11 @@ struct ANIMGRAPHRUNTIME_API FAnimNode_BlendSpaceGraphBase : public FAnimNode_Bas
 
 #if WITH_EDITORONLY_DATA
 	// Set the node to preview a supplied sample value
-	void SetPreviewPosition(FVector InVector);
+	ANIMGRAPHRUNTIME_API void SetPreviewPosition(FVector InVector);
 #endif
 
 	// Forces the Position to the specified value
-	void SnapToPosition(const FVector& NewPosition);
+	ANIMGRAPHRUNTIME_API void SnapToPosition(const FVector& NewPosition);
 
 protected:
 	// The X coordinate to sample in the blendspace
@@ -79,14 +79,14 @@ protected:
 #endif
 
 	// Internal update handler, skipping evaluation of exposed inputs
-	void UpdateInternal(const FAnimationUpdateContext& Context);
+	ANIMGRAPHRUNTIME_API void UpdateInternal(const FAnimationUpdateContext& Context);
 
 protected:
 	// FAnimNode_Base interface
-	virtual void Initialize_AnyThread(const FAnimationInitializeContext& Context) override;
-	virtual void CacheBones_AnyThread(const FAnimationCacheBonesContext& Context) override;
-	virtual void Update_AnyThread(const FAnimationUpdateContext& Context) override;
-	virtual void Evaluate_AnyThread(FPoseContext& Output) override;
-	virtual void GatherDebugData(FNodeDebugData& DebugData) override;
+	ANIMGRAPHRUNTIME_API virtual void Initialize_AnyThread(const FAnimationInitializeContext& Context) override;
+	ANIMGRAPHRUNTIME_API virtual void CacheBones_AnyThread(const FAnimationCacheBonesContext& Context) override;
+	ANIMGRAPHRUNTIME_API virtual void Update_AnyThread(const FAnimationUpdateContext& Context) override;
+	ANIMGRAPHRUNTIME_API virtual void Evaluate_AnyThread(FPoseContext& Output) override;
+	ANIMGRAPHRUNTIME_API virtual void GatherDebugData(FNodeDebugData& DebugData) override;
 	// End of FAnimNode_Base interface
 };
diff --git a/Engine/Source/Runtime/AnimGraphRuntime/Public/AnimNodes/AnimNode_BlendSpacePlayer.h b/Engine/Source/Runtime/AnimGraphRuntime/Public/AnimNodes/AnimNode_BlendSpacePlayer.h
--- a/Engine/Source/Runtime/AnimGraphRuntime/Public/AnimNodes/AnimNode_BlendSpacePlayer.h
+++ b/Engine/Source/Runtime/AnimGraphRuntime/Public/AnimNodes/AnimNode_BlendSpacePlayer.h
@@ -12,7 +12,7 @@ class UBlendSpace;
 
 //@TODO: Comment
 USTRUCT(BlueprintInternalUseOnly)
-struct ANIMGRAPHRUNTIME_API FAnimNode_BlendSpacePlayerBase : public FAnimNode_AssetPlayerBase
+struct FAnimNode_BlendSpacePlayerBase : public FAnimNode_AssetPlayerBase
 {
 	GENERATED_BODY()
 
@@ -32,52 +32,52 @@ protected:
 public:	
 
 	// FAnimNode_AssetPlayerBase interface
-	virtual float GetCurrentAssetTime() const override;
-	virtual float GetCurrentAssetTimePlayRateAdjusted() const override;
-	virtual float GetCurrentAssetLength() const override;
-	virtual UAnimationAsset* GetAnimAsset() const override;
+	ANIMGRAPHRUNTIME_API virtual float GetCurrentAssetTime() const override;
+	ANIMGRAPHRUNTIME_API virtual float GetCurrentAssetTimePlayRateAdjusted() const override;
+	ANIMGRAPHRUNTIME_API virtual float GetCurrentAssetLength() const override;
+	ANIMGRAPHRUNTIME_API virtual UAnimationAsset* GetAnimAsset() const override;
 	// End of FAnimNode_AssetPlayerBase interface
 
 	// FAnimNode_Base interface
-	virtual void Initialize_AnyThread(const FAnimationInitializeContext& Context) override;
-	virtual void CacheBones_AnyThread(const FAnimationCacheBonesContext& Context) override;
-	virtual void UpdateAssetPlayer(const FAnimationUpdateContext& Context) override;
-	virtual void Evaluate_AnyThread(FPoseContext& Output) override;
-	virtual void GatherDebugData(FNodeDebugData& DebugData) override;
+	ANIMGRAPHRUNTIME_API virtual void Initialize_AnyThread(const FAnimationInitializeContext& Context) override;
+	ANIMGRAPHRUNTIME_API virtual void CacheBones_AnyThread(const FAnimationCacheBonesContext& Context) override;
+	ANIMGRAPHRUNTIME_API virtual void UpdateAssetPlayer(const FAnimationUpdateContext& Context) override;
+	ANIMGRAPHRUNTIME_API virtual void Evaluate_AnyThread(FPoseContext& Output) override;
+	ANIMGRAPHRUNTIME_API virtual void GatherDebugData(FNodeDebugData& DebugData) override;
 	// End of FAnimNode_Base interface
 
 	// Get the amount of time from the end
-	float GetTimeFromEnd(float CurrentTime) const;
+	ANIMGRAPHRUNTIME_API float GetTimeFromEnd(float CurrentTime) const;
 
 	// @return the current sample coordinates after going through the filtering
 	FVector GetFilteredPosition() const { return BlendFilter.GetFilterLastOutput(); }
 
 	// Forces the Position to the specified value
-	void SnapToPosition(const FVector& NewPosition);
+	ANIMGRAPHRUNTIME_API void SnapToPosition(const FVector& NewPosition);
 
 public:
 
 	// Get the blendspace asset to play
-	virtual UBlendSpace* GetBlendSpace() const PURE_VIRTUAL(FAnimNode_BlendSpacePlayerBase::GetBlendSpace, return nullptr;);
+	ANIMGRAPHRUNTIME_API virtual UBlendSpace* GetBlendSpace() const PURE_VIRTUAL(FAnimNode_BlendSpacePlayerBase::GetBlendSpace, return nullptr;);
 
 	// Get the coordinates that are currently being sampled by the blendspace
-	virtual FVector GetPosition() const PURE_VIRTUAL(FAnimNode_BlendSpacePlayerBase::GetPosition, return FVector::Zero(););
+	ANIMGRAPHRUNTIME_API virtual FVector GetPosition() const PURE_VIRTUAL(FAnimNode_BlendSpacePlayerBase::GetPosition, return FVector::Zero(););
 
 	// The start position in [0, 1] to use when initializing. When looping, play will still jump back to the beginning when reaching the end.
-	virtual float GetStartPosition() const PURE_VIRTUAL(FAnimNode_BlendSpacePlayerBase::GetStartPosition, return 0.0f;);
+	ANIMGRAPHRUNTIME_API virtual float GetStartPosition() const PURE_VIRTUAL(FAnimNode_BlendSpacePlayerBase::GetStartPosition, return 0.0f;);
 
 	// Get the play rate multiplier. Can be negative, which will cause the animation to play in reverse.
-	virtual float GetPlayRate() const PURE_VIRTUAL(FAnimNode_BlendSpacePlayerBase::GetPlayRate, return 1.0f;);
+	ANIMGRAPHRUNTIME_API virtual float GetPlayRate() const PURE_VIRTUAL(FAnimNode_BlendSpacePlayerBase::GetPlayRate, return 1.0f;);
 
 	// Should the animation loop back to the start when it reaches the end?
 	UE_DEPRECATED(5.3, "Please use IsLooping instead.")
 	virtual bool GetLoop() const final { return IsLooping(); }
 
 	// Get whether we should reset the current play time when the blend space changes
-	virtual bool ShouldResetPlayTimeWhenBlendSpaceChanges() const PURE_VIRTUAL(FAnimNode_BlendSpacePlayerBase::ShouldResetPlayTimeWhenBlendSpaceChanges, return true;);
+	ANIMGRAPHRUNTIME_API virtual bool ShouldResetPlayTimeWhenBlendSpaceChanges() const PURE_VIRTUAL(FAnimNode_BlendSpacePlayerBase::ShouldResetPlayTimeWhenBlendSpaceChanges, return true;);
 
 	// Set whether we should reset the current play time when the blend space changes
-	virtual bool SetResetPlayTimeWhenBlendSpaceChanges(bool bReset) PURE_VIRTUAL(FAnimNode_BlendSpacePlayerBase::SetResetPlayTimeWhenBlendSpaceChanges, return false;);
+	ANIMGRAPHRUNTIME_API virtual bool SetResetPlayTimeWhenBlendSpaceChanges(bool bReset) PURE_VIRTUAL(FAnimNode_BlendSpacePlayerBase::SetResetPlayTimeWhenBlendSpaceChanges, return false;);
 
 	// An evaluator will be setting the play rate to zero and setting the time explicitly. ShouldTeleportToTime indicates whether we should jump to that time, or move to it playing out root motion and events etc.
 	virtual bool ShouldTeleportToTime() const { return false; }
@@ -86,24 +86,24 @@ public:
 	virtual bool IsEvaluator() const { return false; }
 
 	// Set the blendspace asset to play
-	virtual bool SetBlendSpace(UBlendSpace* InBlendSpace) PURE_VIRTUAL(FAnimNode_BlendSpacePlayerBase::SetBlendSpace, return false;);
+	ANIMGRAPHRUNTIME_API virtual bool SetBlendSpace(UBlendSpace* InBlendSpace) PURE_VIRTUAL(FAnimNode_BlendSpacePlayerBase::SetBlendSpace, return false;);
 
 	// Set the coordinates that are currently being sampled by the blendspace
-	virtual bool SetPosition(FVector InPosition) PURE_VIRTUAL(FAnimNode_BlendSpacePlayerBase::SetPosition, return false;);
+	ANIMGRAPHRUNTIME_API virtual bool SetPosition(FVector InPosition) PURE_VIRTUAL(FAnimNode_BlendSpacePlayerBase::SetPosition, return false;);
 
 	// Set the play rate multiplier. Can be negative, which will cause the animation to play in reverse.
-	virtual bool SetPlayRate(float InPlayRate) PURE_VIRTUAL(FAnimNode_BlendSpacePlayerBase::SetPlayRate, return false;);
+	ANIMGRAPHRUNTIME_API virtual bool SetPlayRate(float InPlayRate) PURE_VIRTUAL(FAnimNode_BlendSpacePlayerBase::SetPlayRate, return false;);
 
 	// Set if the animation should loop back to the start when it reaches the end?
-	virtual bool SetLoop(bool bInLoop) PURE_VIRTUAL(FAnimNode_BlendSpacePlayerBase::SetLoop, return false;);
+	ANIMGRAPHRUNTIME_API virtual bool SetLoop(bool bInLoop) PURE_VIRTUAL(FAnimNode_BlendSpacePlayerBase::SetLoop, return false;);
 
 protected:
-	void UpdateInternal(const FAnimationUpdateContext& Context);
+	ANIMGRAPHRUNTIME_API void UpdateInternal(const FAnimationUpdateContext& Context);
 
 private:
-	void Reinitialize(bool bResetTime = true);
+	ANIMGRAPHRUNTIME_API void Reinitialize(bool bResetTime = true);
 
-	const FBlendSampleData* GetHighestWeightedSample() const;
+	ANIMGRAPHRUNTIME_API const FBlendSampleData* GetHighestWeightedSample() const;
 };
 
 template<>
@@ -118,7 +118,7 @@ struct TStructOpsTypeTraits<FAnimNode_BlendSpacePlayerBase> : public TStructOpsT
 
 //@TODO: Comment
 USTRUCT(BlueprintInternalUseOnly)
-struct ANIMGRAPHRUNTIME_API FAnimNode_BlendSpacePlayer : public FAnimNode_BlendSpacePlayerBase
+struct FAnimNode_BlendSpacePlayer : public FAnimNode_BlendSpacePlayerBase
 {
 	GENERATED_BODY()
 
@@ -179,34 +179,34 @@ private:
 public:
 
 	// FAnimNode_AssetPlayerBase interface
-	virtual FName GetGroupName() const override;
-	virtual EAnimGroupRole::Type GetGroupRole() const override;
-	virtual EAnimSyncMethod GetGroupMethod() const override;
-	virtual bool GetIgnoreForRelevancyTest() const override;
-	virtual bool IsLooping() const override;
-	virtual bool SetGroupName(FName InGroupName) override;
-	virtual bool SetGroupRole(EAnimGroupRole::Type InRole) override;
-	virtual bool SetGroupMethod(EAnimSyncMethod InMethod) override;
-	virtual bool SetIgnoreForRelevancyTest(bool bInIgnoreForRelevancyTest) override;
+	ANIMGRAPHRUNTIME_API virtual FName GetGroupName() const override;
+	ANIMGRAPHRUNTIME_API virtual EAnimGroupRole::Type GetGroupRole() const override;
+	ANIMGRAPHRUNTIME_API virtual EAnimSyncMethod GetGroupMethod() const override;
+	ANIMGRAPHRUNTIME_API virtual bool GetIgnoreForRelevancyTest() const override;
+	ANIMGRAPHRUNTIME_API virtual bool IsLooping() const override;
+	ANIMGRAPHRUNTIME_API virtual bool SetGroupName(FName InGroupName) override;
+	ANIMGRAPHRUNTIME_API virtual bool SetGroupRole(EAnimGroupRole::Type InRole) override;
+	ANIMGRAPHRUNTIME_API virtual bool SetGroupMethod(EAnimSyncMethod InMethod) override;
+	ANIMGRAPHRUNTIME_API virtual bool SetIgnoreForRelevancyTest(bool bInIgnoreForRelevancyTest) override;
 	// End of FAnimNode_AssetPlayerBase interface
 
 	// FAnimNode_BlendSpacePlayerBase interface
-	virtual UBlendSpace* GetBlendSpace() const override;
-	virtual FVector GetPosition() const override;
-	virtual float GetStartPosition() const override;
-	virtual float GetPlayRate() const override;
-	virtual bool ShouldResetPlayTimeWhenBlendSpaceChanges() const override;
-	virtual bool SetResetPlayTimeWhenBlendSpaceChanges(bool bReset) override;
-	virtual bool SetBlendSpace(UBlendSpace* InBlendSpace) override;
-	virtual bool SetPosition(FVector InPosition) override;
-	virtual bool SetPlayRate(float InPlayRate) override;
-	virtual bool SetLoop(bool bInLoop) override;
+	ANIMGRAPHRUNTIME_API virtual UBlendSpace* GetBlendSpace() const override;
+	ANIMGRAPHRUNTIME_API virtual FVector GetPosition() const override;
+	ANIMGRAPHRUNTIME_API virtual float GetStartPosition() const override;
+	ANIMGRAPHRUNTIME_API virtual float GetPlayRate() const override;
+	ANIMGRAPHRUNTIME_API virtual bool ShouldResetPlayTimeWhenBlendSpaceChanges() const override;
+	ANIMGRAPHRUNTIME_API virtual bool SetResetPlayTimeWhenBlendSpaceChanges(bool bReset) override;
+	ANIMGRAPHRUNTIME_API virtual bool SetBlendSpace(UBlendSpace* InBlendSpace) override;
+	ANIMGRAPHRUNTIME_API virtual bool SetPosition(FVector InPosition) override;
+	ANIMGRAPHRUNTIME_API virtual bool SetPlayRate(float InPlayRate) override;
+	ANIMGRAPHRUNTIME_API virtual bool SetLoop(bool bInLoop) override;
 	// End of FAnimNode_BlendSpacePlayerBase interface
 };
 
 //@TODO: Comment
 USTRUCT(BlueprintInternalUseOnly)
-struct ANIMGRAPHRUNTIME_API FAnimNode_BlendSpacePlayer_Standalone : public FAnimNode_BlendSpacePlayerBase
+struct FAnimNode_BlendSpacePlayer_Standalone : public FAnimNode_BlendSpacePlayerBase
 {
 	GENERATED_BODY()
 
diff --git a/Engine/Source/Runtime/AnimGraphRuntime/Public/AnimNodes/AnimNode_BlendSpaceSampleResult.h b/Engine/Source/Runtime/AnimGraphRuntime/Public/AnimNodes/AnimNode_BlendSpaceSampleResult.h
--- a/Engine/Source/Runtime/AnimGraphRuntime/Public/AnimNodes/AnimNode_BlendSpaceSampleResult.h
+++ b/Engine/Source/Runtime/AnimGraphRuntime/Public/AnimNodes/AnimNode_BlendSpaceSampleResult.h
@@ -10,7 +10,7 @@
 // Root node of a blend space sample (sink node).
 // We dont use AnimNode_Root to let us distinguish these nodes in the property list at link time.
 USTRUCT(BlueprintInternalUseOnly)
-struct ANIMGRAPHRUNTIME_API FAnimNode_BlendSpaceSampleResult : public FAnimNode_Root
+struct FAnimNode_BlendSpaceSampleResult : public FAnimNode_Root
 {
 	GENERATED_BODY()
 };
diff --git a/Engine/Source/Runtime/AnimGraphRuntime/Public/AnimNodes/AnimNode_CallFunction.h b/Engine/Source/Runtime/AnimGraphRuntime/Public/AnimNodes/AnimNode_CallFunction.h
--- a/Engine/Source/Runtime/AnimGraphRuntime/Public/AnimNodes/AnimNode_CallFunction.h
+++ b/Engine/Source/Runtime/AnimGraphRuntime/Public/AnimNodes/AnimNode_CallFunction.h
@@ -1,4 +1,4 @@
-﻿// Copyright Epic Games, Inc. All Rights Reserved.
+// Copyright Epic Games, Inc. All Rights Reserved.
 #pragma once
 
 #include "CoreMinimal.h"
@@ -53,7 +53,7 @@ enum class EAnimFunctionCallSite
 
 /** Calls specified user-defined events/functions during anim graph execution */
 USTRUCT(BlueprintInternalUseOnly)
-struct ANIMGRAPHRUNTIME_API FAnimNode_CallFunction : public FAnimNode_Base
+struct FAnimNode_CallFunction : public FAnimNode_Base
 {
 	GENERATED_BODY()
 
@@ -77,17 +77,17 @@ struct ANIMGRAPHRUNTIME_API FAnimNode_CallFunction : public FAnimNode_Base
 	EAnimFunctionCallSite CallSite = EAnimFunctionCallSite::OnUpdate;
 		
 	// FAnimNode_Base interface
-	virtual void OnInitializeAnimInstance(const FAnimInstanceProxy* InProxy, const UAnimInstance* InAnimInstance) override;
+	ANIMGRAPHRUNTIME_API virtual void OnInitializeAnimInstance(const FAnimInstanceProxy* InProxy, const UAnimInstance* InAnimInstance) override;
 	virtual bool NeedsOnInitializeAnimInstance() const override { return true; }
-	virtual void GatherDebugData(FNodeDebugData& DebugData) override;
-	virtual void Update_AnyThread(const FAnimationUpdateContext& Context) override;
-	virtual void Evaluate_AnyThread(FPoseContext& Context) override;
-	virtual void Initialize_AnyThread(const FAnimationInitializeContext& Context) override;
+	ANIMGRAPHRUNTIME_API virtual void GatherDebugData(FNodeDebugData& DebugData) override;
+	ANIMGRAPHRUNTIME_API virtual void Update_AnyThread(const FAnimationUpdateContext& Context) override;
+	ANIMGRAPHRUNTIME_API virtual void Evaluate_AnyThread(FPoseContext& Context) override;
+	ANIMGRAPHRUNTIME_API virtual void Initialize_AnyThread(const FAnimationInitializeContext& Context) override;
 	// End of FAnimNode_Base interface
 
 	// Calls the function we hold if the callsite matches the one we have set
-	void CallFunctionFromCallSite(EAnimFunctionCallSite InCallSite, const FAnimationBaseContext& InContext) const;
+	ANIMGRAPHRUNTIME_API void CallFunctionFromCallSite(EAnimFunctionCallSite InCallSite, const FAnimationBaseContext& InContext) const;
 
 	// Get the function held on this node
-	const FAnimNodeFunctionRef& GetFunction() const;
-};
\ No newline at end of file
+	ANIMGRAPHRUNTIME_API const FAnimNodeFunctionRef& GetFunction() const;
+};
diff --git a/Engine/Source/Runtime/AnimGraphRuntime/Public/AnimNodes/AnimNode_CopyPoseFromMesh.h b/Engine/Source/Runtime/AnimGraphRuntime/Public/AnimNodes/AnimNode_CopyPoseFromMesh.h
--- a/Engine/Source/Runtime/AnimGraphRuntime/Public/AnimNodes/AnimNode_CopyPoseFromMesh.h
+++ b/Engine/Source/Runtime/AnimGraphRuntime/Public/AnimNodes/AnimNode_CopyPoseFromMesh.h
@@ -15,7 +15,7 @@ struct FAnimInstanceProxy;
  */
 
 USTRUCT(BlueprintInternalUseOnly)
-struct ANIMGRAPHRUNTIME_API FAnimNode_CopyPoseFromMesh : public FAnimNode_Base
+struct FAnimNode_CopyPoseFromMesh : public FAnimNode_Base
 {
 	GENERATED_USTRUCT_BODY()
 
@@ -43,16 +43,16 @@ struct ANIMGRAPHRUNTIME_API FAnimNode_CopyPoseFromMesh : public FAnimNode_Base
 	UPROPERTY(EditAnywhere, BlueprintReadWrite, Category = Copy, meta = (NeverAsPin))
 	FName RootBoneToCopy;
 
-	FAnimNode_CopyPoseFromMesh();
+	ANIMGRAPHRUNTIME_API FAnimNode_CopyPoseFromMesh();
 
 	// FAnimNode_Base interface
-	virtual void Initialize_AnyThread(const FAnimationInitializeContext& Context) override;
-	virtual void CacheBones_AnyThread(const FAnimationCacheBonesContext& Context) override;
-	virtual void Update_AnyThread(const FAnimationUpdateContext& Context) override;
-	virtual void Evaluate_AnyThread(FPoseContext& Output) override;
-	virtual void GatherDebugData(FNodeDebugData& DebugData) override;
+	ANIMGRAPHRUNTIME_API virtual void Initialize_AnyThread(const FAnimationInitializeContext& Context) override;
+	ANIMGRAPHRUNTIME_API virtual void CacheBones_AnyThread(const FAnimationCacheBonesContext& Context) override;
+	ANIMGRAPHRUNTIME_API virtual void Update_AnyThread(const FAnimationUpdateContext& Context) override;
+	ANIMGRAPHRUNTIME_API virtual void Evaluate_AnyThread(FPoseContext& Output) override;
+	ANIMGRAPHRUNTIME_API virtual void GatherDebugData(FNodeDebugData& DebugData) override;
 	virtual bool HasPreUpdate() const override { return true; }
-	virtual void PreUpdate(const UAnimInstance* InAnimInstance) override;
+	ANIMGRAPHRUNTIME_API virtual void PreUpdate(const UAnimInstance* InAnimInstance) override;
 	// End of FAnimNode_Base interface
 
 private:
@@ -77,6 +77,6 @@ private:
 	UE::Anim::FMeshAttributeContainer SourceCustomAttributes;
 
 	// reinitialize mesh component 
-	void ReinitializeMeshComponent(USkeletalMeshComponent* NewSkeletalMeshComponent, USkeletalMeshComponent* TargetMeshComponent);
-	void RefreshMeshComponent(USkeletalMeshComponent* TargetMeshComponent);
+	ANIMGRAPHRUNTIME_API void ReinitializeMeshComponent(USkeletalMeshComponent* NewSkeletalMeshComponent, USkeletalMeshComponent* TargetMeshComponent);
+	ANIMGRAPHRUNTIME_API void RefreshMeshComponent(USkeletalMeshComponent* TargetMeshComponent);
 };
diff --git a/Engine/Source/Runtime/AnimGraphRuntime/Public/AnimNodes/AnimNode_CurveSource.h b/Engine/Source/Runtime/AnimGraphRuntime/Public/AnimNodes/AnimNode_CurveSource.h
--- a/Engine/Source/Runtime/AnimGraphRuntime/Public/AnimNodes/AnimNode_CurveSource.h
+++ b/Engine/Source/Runtime/AnimGraphRuntime/Public/AnimNodes/AnimNode_CurveSource.h
@@ -14,7 +14,7 @@ class UAnimInstance;
 
 /** Supply curves from some external source (e.g. audio) */
 USTRUCT(BlueprintInternalUseOnly)
-struct ANIMGRAPHRUNTIME_API FAnimNode_CurveSource : public FAnimNode_Base
+struct FAnimNode_CurveSource : public FAnimNode_Base
 {
 	GENERATED_BODY()
 
@@ -38,15 +38,15 @@ struct ANIMGRAPHRUNTIME_API FAnimNode_CurveSource : public FAnimNode_Base
 	UPROPERTY(Transient)
 	TScriptInterface<ICurveSourceInterface> CurveSource;
 
-	FAnimNode_CurveSource();
+	ANIMGRAPHRUNTIME_API FAnimNode_CurveSource();
 
 	// FAnimNode_Base interface
 	virtual bool HasPreUpdate() const override { return true; }
-	virtual void PreUpdate(const UAnimInstance* InAnimInstance) override;
-	virtual void Evaluate_AnyThread(FPoseContext& Output) override;
-	virtual void Update_AnyThread(const FAnimationUpdateContext& Context) override;
-	virtual void Initialize_AnyThread(const FAnimationInitializeContext& Context) override;
-	virtual void CacheBones_AnyThread(const FAnimationCacheBonesContext& Context) override;
-	virtual void GatherDebugData(FNodeDebugData& DebugData) override;
+	ANIMGRAPHRUNTIME_API virtual void PreUpdate(const UAnimInstance* InAnimInstance) override;
+	ANIMGRAPHRUNTIME_API virtual void Evaluate_AnyThread(FPoseContext& Output) override;
+	ANIMGRAPHRUNTIME_API virtual void Update_AnyThread(const FAnimationUpdateContext& Context) override;
+	ANIMGRAPHRUNTIME_API virtual void Initialize_AnyThread(const FAnimationInitializeContext& Context) override;
+	ANIMGRAPHRUNTIME_API virtual void CacheBones_AnyThread(const FAnimationCacheBonesContext& Context) override;
+	ANIMGRAPHRUNTIME_API virtual void GatherDebugData(FNodeDebugData& DebugData) override;
 	// End of FAnimNode_Base interface
 };
diff --git a/Engine/Source/Runtime/AnimGraphRuntime/Public/AnimNodes/AnimNode_LayeredBoneBlend.h b/Engine/Source/Runtime/AnimGraphRuntime/Public/AnimNodes/AnimNode_LayeredBoneBlend.h
--- a/Engine/Source/Runtime/AnimGraphRuntime/Public/AnimNodes/AnimNode_LayeredBoneBlend.h
+++ b/Engine/Source/Runtime/AnimGraphRuntime/Public/AnimNodes/AnimNode_LayeredBoneBlend.h
@@ -18,7 +18,7 @@ enum class ELayeredBoneBlendMode : uint8
 
 // Layered blend (per bone); has dynamic number of blendposes that can blend per different bone sets
 USTRUCT(BlueprintInternalUseOnly)
-struct ANIMGRAPHRUNTIME_API FAnimNode_LayeredBoneBlend : public FAnimNode_Base
+struct FAnimNode_LayeredBoneBlend : public FAnimNode_Base
 {
 	GENERATED_USTRUCT_BODY()
 
@@ -120,11 +120,11 @@ public:
 	}
 
 	// FAnimNode_Base interface
-	virtual void Initialize_AnyThread(const FAnimationInitializeContext& Context) override;
-	virtual void CacheBones_AnyThread(const FAnimationCacheBonesContext& Context) override;
-	virtual void Update_AnyThread(const FAnimationUpdateContext& Context) override;
-	virtual void Evaluate_AnyThread(FPoseContext& Output) override;
-	virtual void GatherDebugData(FNodeDebugData& DebugData) override;
+	ANIMGRAPHRUNTIME_API virtual void Initialize_AnyThread(const FAnimationInitializeContext& Context) override;
+	ANIMGRAPHRUNTIME_API virtual void CacheBones_AnyThread(const FAnimationCacheBonesContext& Context) override;
+	ANIMGRAPHRUNTIME_API virtual void Update_AnyThread(const FAnimationUpdateContext& Context) override;
+	ANIMGRAPHRUNTIME_API virtual void Evaluate_AnyThread(FPoseContext& Output) override;
+	ANIMGRAPHRUNTIME_API virtual void GatherDebugData(FNodeDebugData& DebugData) override;
 	virtual int32 GetLODThreshold() const override { return LODThreshold; }
 	// End of FAnimNode_Base interface
 
@@ -160,7 +160,7 @@ public:
 	}
 
 	// Set the blend mask for the specified input pose
-	void SetBlendMask(int32 InPoseIndex, UBlendProfile* InBlendMask);
+	ANIMGRAPHRUNTIME_API void SetBlendMask(int32 InPoseIndex, UBlendProfile* InBlendMask);
 	
 	// Invalidate the cached per-bone blend weights from the skeleton
 	void InvalidatePerBoneBlendWeights() { RequiredBonesSerialNumber = 0; SkeletonGuid = FGuid(); VirtualBoneGuid = FGuid(); }
@@ -170,13 +170,13 @@ public:
 	
 private:
 	// Rebuild cache per bone blend weights from the skeleton
-	void RebuildPerBoneBlendWeights(const USkeleton* InSkeleton);
+	ANIMGRAPHRUNTIME_API void RebuildPerBoneBlendWeights(const USkeleton* InSkeleton);
 
 	// Check whether per-bone blend weights are valid according to the skeleton (GUID check)
-	bool ArePerBoneBlendWeightsValid(const USkeleton* InSkeleton) const;
+	ANIMGRAPHRUNTIME_API bool ArePerBoneBlendWeightsValid(const USkeleton* InSkeleton) const;
 
 	// Update cached data if required
-	void UpdateCachedBoneData(const FBoneContainer& RequiredBones, const USkeleton* Skeleton);
+	ANIMGRAPHRUNTIME_API void UpdateCachedBoneData(const FBoneContainer& RequiredBones, const USkeleton* Skeleton);
 
 	friend class UAnimGraphNode_LayeredBoneBlend;
 };
diff --git a/Engine/Source/Runtime/AnimGraphRuntime/Public/AnimNodes/AnimNode_MakeDynamicAdditive.h b/Engine/Source/Runtime/AnimGraphRuntime/Public/AnimNodes/AnimNode_MakeDynamicAdditive.h
--- a/Engine/Source/Runtime/AnimGraphRuntime/Public/AnimNodes/AnimNode_MakeDynamicAdditive.h
+++ b/Engine/Source/Runtime/AnimGraphRuntime/Public/AnimNodes/AnimNode_MakeDynamicAdditive.h
@@ -8,7 +8,7 @@
 #include "AnimNode_MakeDynamicAdditive.generated.h"
 
 USTRUCT(BlueprintInternalUseOnly)
-struct ANIMGRAPHRUNTIME_API FAnimNode_MakeDynamicAdditive : public FAnimNode_Base
+struct FAnimNode_MakeDynamicAdditive : public FAnimNode_Base
 {
 	GENERATED_USTRUCT_BODY()
 
@@ -25,14 +25,14 @@ struct ANIMGRAPHRUNTIME_API FAnimNode_MakeDynamicAdditive : public FAnimNode_Bas
 	bool bMeshSpaceAdditive;
 
 public:	
-	FAnimNode_MakeDynamicAdditive();
+	ANIMGRAPHRUNTIME_API FAnimNode_MakeDynamicAdditive();
 
 	// FAnimNode_Base interface
-	virtual void Initialize_AnyThread(const FAnimationInitializeContext& Context) override;
-	virtual void CacheBones_AnyThread(const FAnimationCacheBonesContext& Context) override;
-	virtual void Update_AnyThread(const FAnimationUpdateContext& Context) override;
-	virtual void Evaluate_AnyThread(FPoseContext& Output) override;
-	virtual void GatherDebugData(FNodeDebugData& DebugData) override;
+	ANIMGRAPHRUNTIME_API virtual void Initialize_AnyThread(const FAnimationInitializeContext& Context) override;
+	ANIMGRAPHRUNTIME_API virtual void CacheBones_AnyThread(const FAnimationCacheBonesContext& Context) override;
+	ANIMGRAPHRUNTIME_API virtual void Update_AnyThread(const FAnimationUpdateContext& Context) override;
+	ANIMGRAPHRUNTIME_API virtual void Evaluate_AnyThread(FPoseContext& Output) override;
+	ANIMGRAPHRUNTIME_API virtual void GatherDebugData(FNodeDebugData& DebugData) override;
 	// End of FAnimNode_Base interface
 
 };
diff --git a/Engine/Source/Runtime/AnimGraphRuntime/Public/AnimNodes/AnimNode_Mirror.h b/Engine/Source/Runtime/AnimGraphRuntime/Public/AnimNodes/AnimNode_Mirror.h
--- a/Engine/Source/Runtime/AnimGraphRuntime/Public/AnimNodes/AnimNode_Mirror.h
+++ b/Engine/Source/Runtime/AnimGraphRuntime/Public/AnimNodes/AnimNode_Mirror.h
@@ -9,69 +9,69 @@
 #include "AnimNode_Mirror.generated.h"
 
 USTRUCT(BlueprintInternalUseOnly)
-struct ANIMGRAPHRUNTIME_API FAnimNode_MirrorBase : public FAnimNode_Base
+struct FAnimNode_MirrorBase : public FAnimNode_Base
 {
 	GENERATED_BODY()
 public:
-	FAnimNode_MirrorBase(); 
+	ANIMGRAPHRUNTIME_API FAnimNode_MirrorBase(); 
 
-	virtual void Initialize_AnyThread(const FAnimationInitializeContext& Context) override;
-	virtual void CacheBones_AnyThread(const FAnimationCacheBonesContext& Context) override;
-	virtual void Update_AnyThread(const FAnimationUpdateContext& Context) override;
-	virtual void Evaluate_AnyThread(FPoseContext& Output) override;
-	virtual void GatherDebugData(FNodeDebugData& DebugData) override;
+	ANIMGRAPHRUNTIME_API virtual void Initialize_AnyThread(const FAnimationInitializeContext& Context) override;
+	ANIMGRAPHRUNTIME_API virtual void CacheBones_AnyThread(const FAnimationCacheBonesContext& Context) override;
+	ANIMGRAPHRUNTIME_API virtual void Update_AnyThread(const FAnimationUpdateContext& Context) override;
+	ANIMGRAPHRUNTIME_API virtual void Evaluate_AnyThread(FPoseContext& Output) override;
+	ANIMGRAPHRUNTIME_API virtual void GatherDebugData(FNodeDebugData& DebugData) override;
 
 	// Get the MirrorDataTable
-	virtual UMirrorDataTable* GetMirrorDataTable() const;
+	ANIMGRAPHRUNTIME_API virtual UMirrorDataTable* GetMirrorDataTable() const;
 
 	// Set the MirrorDataTable
-	virtual bool SetMirrorDataTable(UMirrorDataTable* MirrorTable);
+	ANIMGRAPHRUNTIME_API virtual bool SetMirrorDataTable(UMirrorDataTable* MirrorTable);
 
 	// Get Mirror State
-	virtual bool GetMirror() const;
+	ANIMGRAPHRUNTIME_API virtual bool GetMirror() const;
 	// How long to blend using inertialization when switching  mirrored state
-	virtual float GetBlendTimeOnMirrorStateChange() const;
+	ANIMGRAPHRUNTIME_API virtual float GetBlendTimeOnMirrorStateChange() const;
 
 	// Should bones mirror
-	virtual bool GetBoneMirroring() const;
+	ANIMGRAPHRUNTIME_API virtual bool GetBoneMirroring() const;
 
 	// Should the curves mirror
-	virtual bool GetCurveMirroring() const;
+	ANIMGRAPHRUNTIME_API virtual bool GetCurveMirroring() const;
 
 	// Should attributes mirror (based on the bone mirroring data in the mirror data table) 
-	virtual bool GetAttributeMirroring() const;
+	ANIMGRAPHRUNTIME_API virtual bool GetAttributeMirroring() const;
 
 	// Whether to reset (reinitialize) the child (source) pose when the mirror state changes
-	virtual bool GetResetChildOnMirrorStateChange() const;
+	ANIMGRAPHRUNTIME_API virtual bool GetResetChildOnMirrorStateChange() const;
 
 	// Set Mirror State
-	virtual bool SetMirror(bool bInMirror);
+	ANIMGRAPHRUNTIME_API virtual bool SetMirror(bool bInMirror);
 
 	// Set how long to blend using inertialization when switching  mirrored state
 	// @return true if the value was set (it is dynamic), or false if it could not (it is not dynamic or pin exposed)
-	virtual bool SetBlendTimeOnMirrorStateChange(float InBlendTime);
+	ANIMGRAPHRUNTIME_API virtual bool SetBlendTimeOnMirrorStateChange(float InBlendTime);
 
 	// Set if bones mirror
 	// @return true if the value was set (it is dynamic), or false if it could not (it is not dynamic or pin exposed)
-	virtual bool SetBoneMirroring(bool bInBoneMirroring);
+	ANIMGRAPHRUNTIME_API virtual bool SetBoneMirroring(bool bInBoneMirroring);
 
 	// Set if curves mirror
 	// @return true if the value was set (it is dynamic), or false if it could not (it is not dynamic or pin exposed)
-	virtual bool SetCurveMirroring(bool bInCurveMirroring);
+	ANIMGRAPHRUNTIME_API virtual bool SetCurveMirroring(bool bInCurveMirroring);
 
 	// Set if attributes mirror
 	// @return true if the value was set (it is dynamic), or false if it could not (it is not dynamic or pin exposed)
-	virtual bool SetAttributeMirroring(bool bInAttributeMirroring);
+	ANIMGRAPHRUNTIME_API virtual bool SetAttributeMirroring(bool bInAttributeMirroring);
 
 	// Set whether to reset (reinitialize) the child (source) pose when the mirror state changes
 	// @return true if the value was set (it is dynamic), or false if it could not (it is not dynamic or pin exposed)
-	virtual bool SetResetChildOnMirrorStateChange(bool bInResetChildOnMirrorStateChange);
+	ANIMGRAPHRUNTIME_API virtual bool SetResetChildOnMirrorStateChange(bool bInResetChildOnMirrorStateChange);
 
 	/** This only used by custom handlers, and it is advanced feature. */
-	virtual void SetSourceLinkNode(FAnimNode_Base* NewLinkNode);
+	ANIMGRAPHRUNTIME_API virtual void SetSourceLinkNode(FAnimNode_Base* NewLinkNode);
 
 	/** This only used by custom handlers, and it is advanced feature. */
-	virtual FAnimNode_Base* GetSourceLinkNode();
+	ANIMGRAPHRUNTIME_API virtual FAnimNode_Base* GetSourceLinkNode();
 protected:
 	UPROPERTY(EditAnywhere, Category = Links)
 	FPoseLink Source;
@@ -90,31 +90,31 @@ private:
 
 
 USTRUCT(BlueprintInternalUseOnly)
-struct ANIMGRAPHRUNTIME_API FAnimNode_Mirror : public FAnimNode_MirrorBase
+struct FAnimNode_Mirror : public FAnimNode_MirrorBase
 {
 	GENERATED_BODY()
 
 	friend class UAnimGraphNode_Mirror;
 
 public:
-	FAnimNode_Mirror();
+	ANIMGRAPHRUNTIME_API FAnimNode_Mirror();
 
-	virtual UMirrorDataTable* GetMirrorDataTable() const override;
-	virtual bool SetMirrorDataTable(UMirrorDataTable* MirrorTable) override;
+	ANIMGRAPHRUNTIME_API virtual UMirrorDataTable* GetMirrorDataTable() const override;
+	ANIMGRAPHRUNTIME_API virtual bool SetMirrorDataTable(UMirrorDataTable* MirrorTable) override;
 
-	virtual bool GetMirror() const override;
-	virtual float GetBlendTimeOnMirrorStateChange() const override;
-	virtual bool GetBoneMirroring() const override;
-	virtual bool GetCurveMirroring() const override;
-	virtual bool GetAttributeMirroring() const override;
-	virtual bool GetResetChildOnMirrorStateChange() const override;
+	ANIMGRAPHRUNTIME_API virtual bool GetMirror() const override;
+	ANIMGRAPHRUNTIME_API virtual float GetBlendTimeOnMirrorStateChange() const override;
+	ANIMGRAPHRUNTIME_API virtual bool GetBoneMirroring() const override;
+	ANIMGRAPHRUNTIME_API virtual bool GetCurveMirroring() const override;
+	ANIMGRAPHRUNTIME_API virtual bool GetAttributeMirroring() const override;
+	ANIMGRAPHRUNTIME_API virtual bool GetResetChildOnMirrorStateChange() const override;
 
-	virtual bool SetMirror(bool bInMirror) override;
-	virtual bool SetBlendTimeOnMirrorStateChange(float InBlendTime) override;
-	virtual bool SetBoneMirroring(bool bInBoneMirroring) override;
-	virtual bool SetCurveMirroring(bool bInCurveMirroring) override;
-	virtual bool SetAttributeMirroring(bool bInAttributeMirroring) override;
-	virtual bool SetResetChildOnMirrorStateChange(bool bInResetChildOnMirrorStateChange) override;
+	ANIMGRAPHRUNTIME_API virtual bool SetMirror(bool bInMirror) override;
+	ANIMGRAPHRUNTIME_API virtual bool SetBlendTimeOnMirrorStateChange(float InBlendTime) override;
+	ANIMGRAPHRUNTIME_API virtual bool SetBoneMirroring(bool bInBoneMirroring) override;
+	ANIMGRAPHRUNTIME_API virtual bool SetCurveMirroring(bool bInCurveMirroring) override;
+	ANIMGRAPHRUNTIME_API virtual bool SetAttributeMirroring(bool bInAttributeMirroring) override;
+	ANIMGRAPHRUNTIME_API virtual bool SetResetChildOnMirrorStateChange(bool bInResetChildOnMirrorStateChange) override;
 
 protected:
 
@@ -145,31 +145,31 @@ protected:
 
 
 USTRUCT(BlueprintInternalUseOnly)
-struct ANIMGRAPHRUNTIME_API FAnimNode_Mirror_Standalone : public FAnimNode_MirrorBase
+struct FAnimNode_Mirror_Standalone : public FAnimNode_MirrorBase
 {
 	GENERATED_BODY()
 
 	friend class UAnimGraphNode_Mirror;
 
 public:
-	FAnimNode_Mirror_Standalone();
+	ANIMGRAPHRUNTIME_API FAnimNode_Mirror_Standalone();
 
-	virtual UMirrorDataTable* GetMirrorDataTable() const override;
-	virtual bool SetMirrorDataTable(UMirrorDataTable* MirrorTable) override;
+	ANIMGRAPHRUNTIME_API virtual UMirrorDataTable* GetMirrorDataTable() const override;
+	ANIMGRAPHRUNTIME_API virtual bool SetMirrorDataTable(UMirrorDataTable* MirrorTable) override;
 
-	virtual bool GetMirror() const override;
-	virtual float GetBlendTimeOnMirrorStateChange() const override;
-	virtual bool GetBoneMirroring() const override;
-	virtual bool GetCurveMirroring() const override;
-	virtual bool GetAttributeMirroring() const override;
-	virtual bool GetResetChildOnMirrorStateChange() const override;
+	ANIMGRAPHRUNTIME_API virtual bool GetMirror() const override;
+	ANIMGRAPHRUNTIME_API virtual float GetBlendTimeOnMirrorStateChange() const override;
+	ANIMGRAPHRUNTIME_API virtual bool GetBoneMirroring() const override;
+	ANIMGRAPHRUNTIME_API virtual bool GetCurveMirroring() const override;
+	ANIMGRAPHRUNTIME_API virtual bool GetAttributeMirroring() const override;
+	ANIMGRAPHRUNTIME_API virtual bool GetResetChildOnMirrorStateChange() const override;
 
-	virtual bool SetMirror(bool bInMirror) override;
-	virtual bool SetBlendTimeOnMirrorStateChange(float InBlendTime) override;
-	virtual bool SetBoneMirroring(bool bInBoneMirroring) override;
-	virtual bool SetCurveMirroring(bool bInCurveMirroring) override;
-	virtual bool SetAttributeMirroring(bool bInAttributeMirroring) override;
-	virtual bool SetResetChildOnMirrorStateChange(bool bInResetChildOnMirrorStateChange) override;
+	ANIMGRAPHRUNTIME_API virtual bool SetMirror(bool bInMirror) override;
+	ANIMGRAPHRUNTIME_API virtual bool SetBlendTimeOnMirrorStateChange(float InBlendTime) override;
+	ANIMGRAPHRUNTIME_API virtual bool SetBoneMirroring(bool bInBoneMirroring) override;
+	ANIMGRAPHRUNTIME_API virtual bool SetCurveMirroring(bool bInCurveMirroring) override;
+	ANIMGRAPHRUNTIME_API virtual bool SetAttributeMirroring(bool bInAttributeMirroring) override;
+	ANIMGRAPHRUNTIME_API virtual bool SetResetChildOnMirrorStateChange(bool bInResetChildOnMirrorStateChange) override;
 
 protected:
 
@@ -195,4 +195,4 @@ protected:
 
 	UPROPERTY(EditAnywhere, Category = MirroredChannels, meta = (DisplayName = "Attributes", FoldProperty))
 	bool bAttributeMirroring = true;
-};
\ No newline at end of file
+};
diff --git a/Engine/Source/Runtime/AnimGraphRuntime/Public/AnimNodes/AnimNode_ModifyCurve.h b/Engine/Source/Runtime/AnimGraphRuntime/Public/AnimNodes/AnimNode_ModifyCurve.h
--- a/Engine/Source/Runtime/AnimGraphRuntime/Public/AnimNodes/AnimNode_ModifyCurve.h
+++ b/Engine/Source/Runtime/AnimGraphRuntime/Public/AnimNodes/AnimNode_ModifyCurve.h
@@ -28,7 +28,7 @@ enum class EModifyCurveApplyMode : uint8
 
 /** Easy way to modify curve values on a pose */
 USTRUCT(BlueprintInternalUseOnly)
-struct ANIMGRAPHRUNTIME_API FAnimNode_ModifyCurve : public FAnimNode_Base
+struct FAnimNode_ModifyCurve : public FAnimNode_Base
 {
 	GENERATED_BODY()
 
@@ -61,28 +61,28 @@ struct ANIMGRAPHRUNTIME_API FAnimNode_ModifyCurve : public FAnimNode_Base
 	UPROPERTY(EditAnywhere, BlueprintReadWrite, Category = ModifyCurve)
 	EModifyCurveApplyMode ApplyMode;
 
-	FAnimNode_ModifyCurve();
+	ANIMGRAPHRUNTIME_API FAnimNode_ModifyCurve();
 	PRAGMA_DISABLE_DEPRECATION_WARNINGS
 	FAnimNode_ModifyCurve(const FAnimNode_ModifyCurve&) = default;
 	FAnimNode_ModifyCurve& operator=(const FAnimNode_ModifyCurve&) = default;
-	PRAGMA_ENABLE_DEPRECATION_WARNINGS
+	ANIMGRAPHRUNTIME_API PRAGMA_ENABLE_DEPRECATION_WARNINGS
 
 	// FAnimNode_Base interface
 	virtual void Initialize_AnyThread(const FAnimationInitializeContext& Context) override;
-	virtual void CacheBones_AnyThread(const FAnimationCacheBonesContext& Context) override;
-	virtual void Evaluate_AnyThread(FPoseContext& Output) override;
+	ANIMGRAPHRUNTIME_API virtual void CacheBones_AnyThread(const FAnimationCacheBonesContext& Context) override;
+	ANIMGRAPHRUNTIME_API virtual void Evaluate_AnyThread(FPoseContext& Output) override;
 
-	virtual void Update_AnyThread(const FAnimationUpdateContext& Context) override;
+	ANIMGRAPHRUNTIME_API virtual void Update_AnyThread(const FAnimationUpdateContext& Context) override;
 	// End of FAnimNode_Base interface
 
 #if WITH_EDITOR
 	/** Add new curve being modified */
-	void AddCurve(const FName& InName, float InValue);
+	ANIMGRAPHRUNTIME_API void AddCurve(const FName& InName, float InValue);
 	/** Remove a curve from being modified */
-	void RemoveCurve(int32 PoseIndex);
+	ANIMGRAPHRUNTIME_API void RemoveCurve(int32 PoseIndex);
 #endif // WITH_EDITOR
 
 private:
-	float ProcessCurveOperation(float CurrentValue, float NewValue) const;
-	float ProcessCurveWMAOperation(float CurrentValue, float LastValue) const;
+	ANIMGRAPHRUNTIME_API float ProcessCurveOperation(float CurrentValue, float NewValue) const;
+	ANIMGRAPHRUNTIME_API float ProcessCurveWMAOperation(float CurrentValue, float LastValue) const;
 };
diff --git a/Engine/Source/Runtime/AnimGraphRuntime/Public/AnimNodes/AnimNode_MultiWayBlend.h b/Engine/Source/Runtime/AnimGraphRuntime/Public/AnimNodes/AnimNode_MultiWayBlend.h
--- a/Engine/Source/Runtime/AnimGraphRuntime/Public/AnimNodes/AnimNode_MultiWayBlend.h
+++ b/Engine/Source/Runtime/AnimGraphRuntime/Public/AnimNodes/AnimNode_MultiWayBlend.h
@@ -8,7 +8,7 @@
 
 // This represents a baked transition
 USTRUCT(BlueprintInternalUseOnly)
-struct ANIMGRAPHRUNTIME_API FAnimNode_MultiWayBlend : public FAnimNode_Base
+struct FAnimNode_MultiWayBlend : public FAnimNode_Base
 {
 	GENERATED_USTRUCT_BODY()
 
@@ -41,11 +41,11 @@ public:
 	}
 
 	// FAnimNode_Base interface
-	virtual void Initialize_AnyThread(const FAnimationInitializeContext& Context) override;
-	virtual void CacheBones_AnyThread(const FAnimationCacheBonesContext& Context) override;
-	virtual void Update_AnyThread(const FAnimationUpdateContext& Context) override;
-	virtual void Evaluate_AnyThread(FPoseContext& Output) override;
-	virtual void GatherDebugData(FNodeDebugData& DebugData) override;
+	ANIMGRAPHRUNTIME_API virtual void Initialize_AnyThread(const FAnimationInitializeContext& Context) override;
+	ANIMGRAPHRUNTIME_API virtual void CacheBones_AnyThread(const FAnimationCacheBonesContext& Context) override;
+	ANIMGRAPHRUNTIME_API virtual void Update_AnyThread(const FAnimationUpdateContext& Context) override;
+	ANIMGRAPHRUNTIME_API virtual void Evaluate_AnyThread(FPoseContext& Output) override;
+	ANIMGRAPHRUNTIME_API virtual void GatherDebugData(FNodeDebugData& DebugData) override;
 	// End of FAnimNode_Base interface
 
 	int32 AddPose()
@@ -85,6 +85,6 @@ public:
 
 private:
 	// process new weights and then return out
-	void UpdateCachedAlphas();
+	ANIMGRAPHRUNTIME_API void UpdateCachedAlphas();
 };
 
diff --git a/Engine/Source/Runtime/AnimGraphRuntime/Public/AnimNodes/AnimNode_PoseBlendNode.h b/Engine/Source/Runtime/AnimGraphRuntime/Public/AnimNodes/AnimNode_PoseBlendNode.h
--- a/Engine/Source/Runtime/AnimGraphRuntime/Public/AnimNodes/AnimNode_PoseBlendNode.h
+++ b/Engine/Source/Runtime/AnimGraphRuntime/Public/AnimNodes/AnimNode_PoseBlendNode.h
@@ -13,7 +13,7 @@
 // Typically the playback position of the animation for this node will represent something other than time, like jump height.
 // This node will not trigger any notifies present in the associated sequence.
 USTRUCT(BlueprintInternalUseOnly)
-struct ANIMGRAPHRUNTIME_API FAnimNode_PoseBlendNode : public FAnimNode_PoseHandler
+struct FAnimNode_PoseBlendNode : public FAnimNode_PoseHandler
 {
 	GENERATED_USTRUCT_BODY()
 public:
@@ -33,17 +33,17 @@ private:
 	UE::Anim::TNamedValueArray<FDefaultAllocator, UE::Anim::FNamedIndexElement> BulkCurves;
 
 public:	
-	FAnimNode_PoseBlendNode();
+	ANIMGRAPHRUNTIME_API FAnimNode_PoseBlendNode();
 
 	// FAnimNode_Base interface
-	virtual void Initialize_AnyThread(const FAnimationInitializeContext& Context) override;
-	virtual void CacheBones_AnyThread(const FAnimationCacheBonesContext& Context) override;
-	virtual void UpdateAssetPlayer(const FAnimationUpdateContext& Context) override;
-	virtual void Evaluate_AnyThread(FPoseContext& Output) override;
-	virtual void GatherDebugData(FNodeDebugData& DebugData) override;
+	ANIMGRAPHRUNTIME_API virtual void Initialize_AnyThread(const FAnimationInitializeContext& Context) override;
+	ANIMGRAPHRUNTIME_API virtual void CacheBones_AnyThread(const FAnimationCacheBonesContext& Context) override;
+	ANIMGRAPHRUNTIME_API virtual void UpdateAssetPlayer(const FAnimationUpdateContext& Context) override;
+	ANIMGRAPHRUNTIME_API virtual void Evaluate_AnyThread(FPoseContext& Output) override;
+	ANIMGRAPHRUNTIME_API virtual void GatherDebugData(FNodeDebugData& DebugData) override;
 	// End of FAnimNode_Base interface
 
 	// FAnimNode_PoseHandler interface 
-	virtual void RebuildPoseList(const FBoneContainer& InBoneContainer, const UPoseAsset* InPoseAsset) override;
+	ANIMGRAPHRUNTIME_API virtual void RebuildPoseList(const FBoneContainer& InBoneContainer, const UPoseAsset* InPoseAsset) override;
 };
 
diff --git a/Engine/Source/Runtime/AnimGraphRuntime/Public/AnimNodes/AnimNode_PoseByName.h b/Engine/Source/Runtime/AnimGraphRuntime/Public/AnimNodes/AnimNode_PoseByName.h
--- a/Engine/Source/Runtime/AnimGraphRuntime/Public/AnimNodes/AnimNode_PoseByName.h
+++ b/Engine/Source/Runtime/AnimGraphRuntime/Public/AnimNodes/AnimNode_PoseByName.h
@@ -11,7 +11,7 @@
 // Typically the playback position of the animation for this node will represent something other than time, like jump height.
 // This node will not trigger any notifies present in the associated sequence.
 USTRUCT(BlueprintInternalUseOnly)
-struct ANIMGRAPHRUNTIME_API FAnimNode_PoseByName : public FAnimNode_PoseHandler
+struct FAnimNode_PoseByName : public FAnimNode_PoseHandler
 {
 	GENERATED_USTRUCT_BODY()
 public:
@@ -33,12 +33,12 @@ public:
 	}
 
 	// FAnimNode_Base interface
-	virtual void Initialize_AnyThread(const FAnimationInitializeContext& Context) override;
-	virtual void UpdateAssetPlayer(const FAnimationUpdateContext& Context) override;
-	virtual void Evaluate_AnyThread(FPoseContext& Output) override;
-	virtual void GatherDebugData(FNodeDebugData& DebugData) override;
+	ANIMGRAPHRUNTIME_API virtual void Initialize_AnyThread(const FAnimationInitializeContext& Context) override;
+	ANIMGRAPHRUNTIME_API virtual void UpdateAssetPlayer(const FAnimationUpdateContext& Context) override;
+	ANIMGRAPHRUNTIME_API virtual void Evaluate_AnyThread(FPoseContext& Output) override;
+	ANIMGRAPHRUNTIME_API virtual void GatherDebugData(FNodeDebugData& DebugData) override;
 	// End of FAnimNode_Base interface
 private:
-	virtual void RebuildPoseList(const FBoneContainer& InBoneContainer, const UPoseAsset* InPoseAsset) override;
+	ANIMGRAPHRUNTIME_API virtual void RebuildPoseList(const FBoneContainer& InBoneContainer, const UPoseAsset* InPoseAsset) override;
 };
 
diff --git a/Engine/Source/Runtime/AnimGraphRuntime/Public/AnimNodes/AnimNode_PoseDriver.h b/Engine/Source/Runtime/AnimGraphRuntime/Public/AnimNodes/AnimNode_PoseDriver.h
--- a/Engine/Source/Runtime/AnimGraphRuntime/Public/AnimNodes/AnimNode_PoseDriver.h
+++ b/Engine/Source/Runtime/AnimGraphRuntime/Public/AnimNodes/AnimNode_PoseDriver.h
@@ -63,7 +63,7 @@ struct FPoseDriverTransform
 
 /** Information about each target in the PoseDriver */
 USTRUCT()
-struct ANIMGRAPHRUNTIME_API FPoseDriverTarget
+struct FPoseDriverTarget
 {
 	GENERATED_BODY()
 		
@@ -133,7 +133,7 @@ struct ANIMGRAPHRUNTIME_API FPoseDriverTarget
 
 /** RBF based orientation driver */
 USTRUCT(BlueprintInternalUseOnly)
-struct ANIMGRAPHRUNTIME_API FAnimNode_PoseDriver : public FAnimNode_PoseHandler
+struct FAnimNode_PoseDriver : public FAnimNode_PoseHandler
 {
 	GENERATED_BODY()
 
@@ -230,27 +230,27 @@ struct ANIMGRAPHRUNTIME_API FAnimNode_PoseDriver : public FAnimNode_PoseHandler
 	int32 LODThreshold;
 
 	// FAnimNode_Base interface
-	virtual void Initialize_AnyThread(const FAnimationInitializeContext& Context) override;
-	virtual void CacheBones_AnyThread(const FAnimationCacheBonesContext& Context) override;
-	virtual void UpdateAssetPlayer(const FAnimationUpdateContext& Context) override;
-	virtual void Evaluate_AnyThread(FPoseContext& Output) override;
-	virtual void GatherDebugData(FNodeDebugData& DebugData) override;
+	ANIMGRAPHRUNTIME_API virtual void Initialize_AnyThread(const FAnimationInitializeContext& Context) override;
+	ANIMGRAPHRUNTIME_API virtual void CacheBones_AnyThread(const FAnimationCacheBonesContext& Context) override;
+	ANIMGRAPHRUNTIME_API virtual void UpdateAssetPlayer(const FAnimationUpdateContext& Context) override;
+	ANIMGRAPHRUNTIME_API virtual void Evaluate_AnyThread(FPoseContext& Output) override;
+	ANIMGRAPHRUNTIME_API virtual void GatherDebugData(FNodeDebugData& DebugData) override;
 	virtual int32 GetLODThreshold() const override { return LODThreshold; }
 	// End of FAnimNode_Base interface
 
-	FAnimNode_PoseDriver();
+	ANIMGRAPHRUNTIME_API FAnimNode_PoseDriver();
 
 	/** Returns the radius for a given target */
-	float GetRadiusForTarget(const FRBFTarget& Target) const;
+	ANIMGRAPHRUNTIME_API float GetRadiusForTarget(const FRBFTarget& Target) const;
 
 	/** Util for seeing if BoneName is in the list of driven bones (and bFilterDrivenBones is true) */
-	bool IsBoneDriven(FName BoneName) const;
+	ANIMGRAPHRUNTIME_API bool IsBoneDriven(FName BoneName) const;
 
 	/** Return array of FRBFTarget structs, derived from PoseTargets array and DriveSource setting */
-	void GetRBFTargets(TArray<FRBFTarget>& OutTargets, const FBoneContainer* BoneContainer) const;
+	ANIMGRAPHRUNTIME_API void GetRBFTargets(TArray<FRBFTarget>& OutTargets, const FBoneContainer* BoneContainer) const;
 
 	/* Rebuild Pose List*/
-	virtual void RebuildPoseList(const FBoneContainer& InBoneContainer, const UPoseAsset* InPoseAsset) override;
+	ANIMGRAPHRUNTIME_API virtual void RebuildPoseList(const FBoneContainer& InBoneContainer, const UPoseAsset* InPoseAsset) override;
 
 private:
 	TSharedPtr<const FRBFSolverData> SolverData;
diff --git a/Engine/Source/Runtime/AnimGraphRuntime/Public/AnimNodes/AnimNode_PoseHandler.h b/Engine/Source/Runtime/AnimGraphRuntime/Public/AnimNodes/AnimNode_PoseHandler.h
--- a/Engine/Source/Runtime/AnimGraphRuntime/Public/AnimNodes/AnimNode_PoseHandler.h
+++ b/Engine/Source/Runtime/AnimGraphRuntime/Public/AnimNodes/AnimNode_PoseHandler.h
@@ -13,7 +13,7 @@
 // Typically the playback position of the animation for this node will represent something other than time, like jump height.
 // This node will not trigger any notifies present in the associated sequence.
 USTRUCT(BlueprintInternalUseOnly)
-struct ANIMGRAPHRUNTIME_API FAnimNode_PoseHandler : public FAnimNode_AssetPlayerBase
+struct FAnimNode_PoseHandler : public FAnimNode_AssetPlayerBase
 {
 	GENERATED_USTRUCT_BODY()
 public:
@@ -33,10 +33,10 @@ public:
 	// End of FAnimNode_AssetPlayerBase interface
 
 	// FAnimNode_Base interface
-	virtual void Initialize_AnyThread(const FAnimationInitializeContext& Context) override;
-	virtual void CacheBones_AnyThread(const FAnimationCacheBonesContext& Context) override;
-	virtual void UpdateAssetPlayer(const FAnimationUpdateContext& Context) override;
-	virtual void GatherDebugData(FNodeDebugData& DebugData) override;
+	ANIMGRAPHRUNTIME_API virtual void Initialize_AnyThread(const FAnimationInitializeContext& Context) override;
+	ANIMGRAPHRUNTIME_API virtual void CacheBones_AnyThread(const FAnimationCacheBonesContext& Context) override;
+	ANIMGRAPHRUNTIME_API virtual void UpdateAssetPlayer(const FAnimationUpdateContext& Context) override;
+	ANIMGRAPHRUNTIME_API virtual void GatherDebugData(FNodeDebugData& DebugData) override;
 	// End of FAnimNode_Base interface
 
 	// FAnimNode_AssetPlayerBase Interface
@@ -47,7 +47,7 @@ public:
 
 #if WITH_EDITORONLY_DATA
 	// Set the pose asset to use for this node 
-	void SetPoseAsset(UPoseAsset* InPoseAsset);
+	ANIMGRAPHRUNTIME_API void SetPoseAsset(UPoseAsset* InPoseAsset);
 #endif
 	
 protected:
@@ -61,12 +61,12 @@ protected:
 	TArray<float> BoneBlendWeights;
 
 	/* Rebuild pose list */
-	virtual void RebuildPoseList(const FBoneContainer& InBoneContainer, const UPoseAsset* InPoseAsset);
+	ANIMGRAPHRUNTIME_API virtual void RebuildPoseList(const FBoneContainer& InBoneContainer, const UPoseAsset* InPoseAsset);
 
 	/** Cache bone blend weights - called when pose asset changes */
-	void CacheBoneBlendWeights(FAnimInstanceProxy* InstanceProxy);
+	ANIMGRAPHRUNTIME_API void CacheBoneBlendWeights(FAnimInstanceProxy* InstanceProxy);
 	
 private:
-	void UpdatePoseAssetProperty(struct FAnimInstanceProxy* InstanceProxy);
+	ANIMGRAPHRUNTIME_API void UpdatePoseAssetProperty(struct FAnimInstanceProxy* InstanceProxy);
 };
 
diff --git a/Engine/Source/Runtime/AnimGraphRuntime/Public/AnimNodes/AnimNode_PoseSnapshot.h b/Engine/Source/Runtime/AnimGraphRuntime/Public/AnimNodes/AnimNode_PoseSnapshot.h
--- a/Engine/Source/Runtime/AnimGraphRuntime/Public/AnimNodes/AnimNode_PoseSnapshot.h
+++ b/Engine/Source/Runtime/AnimGraphRuntime/Public/AnimNodes/AnimNode_PoseSnapshot.h
@@ -29,20 +29,20 @@ enum class ESnapshotSourceMode : uint8
 
 /** Provide a snapshot pose, either from the internal named pose cache or via a supplied snapshot */
 USTRUCT(BlueprintInternalUseOnly)
-struct ANIMGRAPHRUNTIME_API FAnimNode_PoseSnapshot : public FAnimNode_Base
+struct FAnimNode_PoseSnapshot : public FAnimNode_Base
 {
 	GENERATED_USTRUCT_BODY()
 
 public:	
 
-	FAnimNode_PoseSnapshot();
+	ANIMGRAPHRUNTIME_API FAnimNode_PoseSnapshot();
 
 	/** FAnimNode_Base interface */
 	virtual bool HasPreUpdate() const override  { return true; }
-	virtual void PreUpdate(const UAnimInstance* InAnimInstance) override;
-	virtual void Update_AnyThread(const FAnimationUpdateContext& Context) override;
-	virtual void Evaluate_AnyThread(FPoseContext& Output) override;
-	virtual void GatherDebugData(FNodeDebugData& DebugData) override;
+	ANIMGRAPHRUNTIME_API virtual void PreUpdate(const UAnimInstance* InAnimInstance) override;
+	ANIMGRAPHRUNTIME_API virtual void Update_AnyThread(const FAnimationUpdateContext& Context) override;
+	ANIMGRAPHRUNTIME_API virtual void Evaluate_AnyThread(FPoseContext& Output) override;
+	ANIMGRAPHRUNTIME_API virtual void GatherDebugData(FNodeDebugData& DebugData) override;
 
 	/** The name of the snapshot previously stored with SavePoseSnapshot */
 	UPROPERTY(EditAnywhere, BlueprintReadWrite, Category = "Snapshot", meta = (PinShownByDefault))
@@ -72,8 +72,8 @@ private:
 
 private:
 	/** Evaluation helper function - apply a snapshot pose to a pose */
-	void ApplyPose(const FPoseSnapshot& PoseSnapshot, FCompactPose& OutPose);
+	ANIMGRAPHRUNTIME_API void ApplyPose(const FPoseSnapshot& PoseSnapshot, FCompactPose& OutPose);
 
 	/** Evaluation helper function - cache the bone mapping between two skeletal meshes */
-	void CacheBoneMapping(FName SourceMeshName, FName TargetMeshName, const TArray<FName>& InSourceBoneNames, const TArray<FName>& InTargetBoneNames);
+	ANIMGRAPHRUNTIME_API void CacheBoneMapping(FName SourceMeshName, FName TargetMeshName, const TArray<FName>& InSourceBoneNames, const TArray<FName>& InTargetBoneNames);
 };
diff --git a/Engine/Source/Runtime/AnimGraphRuntime/Public/AnimNodes/AnimNode_RandomPlayer.h b/Engine/Source/Runtime/AnimGraphRuntime/Public/AnimNodes/AnimNode_RandomPlayer.h
--- a/Engine/Source/Runtime/AnimGraphRuntime/Public/AnimNodes/AnimNode_RandomPlayer.h
+++ b/Engine/Source/Runtime/AnimGraphRuntime/Public/AnimNodes/AnimNode_RandomPlayer.h
@@ -93,11 +93,11 @@ struct FRandomAnimPlayData
 };
 
 USTRUCT(BlueprintInternalUseOnly)
-struct ANIMGRAPHRUNTIME_API FAnimNode_RandomPlayer : public FAnimNode_AssetPlayerRelevancyBase
+struct FAnimNode_RandomPlayer : public FAnimNode_AssetPlayerRelevancyBase
 {
 	GENERATED_BODY()
 
-	FAnimNode_RandomPlayer();
+	ANIMGRAPHRUNTIME_API FAnimNode_RandomPlayer();
 
 public:
 	/** List of sequences to randomly step through */
@@ -105,20 +105,20 @@ public:
 	TArray<FRandomPlayerSequenceEntry> Entries;
 
 	// FAnimNode_Base interface
-	virtual void Initialize_AnyThread(const FAnimationInitializeContext& Context) override;
-	virtual void Update_AnyThread(const FAnimationUpdateContext& Context) override;
-	virtual void Evaluate_AnyThread(FPoseContext& Output) override;
-	virtual void GatherDebugData(FNodeDebugData& DebugData) override;
+	ANIMGRAPHRUNTIME_API virtual void Initialize_AnyThread(const FAnimationInitializeContext& Context) override;
+	ANIMGRAPHRUNTIME_API virtual void Update_AnyThread(const FAnimationUpdateContext& Context) override;
+	ANIMGRAPHRUNTIME_API virtual void Evaluate_AnyThread(FPoseContext& Output) override;
+	ANIMGRAPHRUNTIME_API virtual void GatherDebugData(FNodeDebugData& DebugData) override;
 	// End of FAnimNode_Base interface
 
 	// FAnimNode_RelevantAssetPlayerBase
-	virtual UAnimationAsset* GetAnimAsset() const override;
-	virtual float GetAccumulatedTime() const override;
-	virtual bool GetIgnoreForRelevancyTest() const override;
-	virtual bool SetIgnoreForRelevancyTest(bool bInIgnoreForRelevancyTest) override;
-	virtual float GetCachedBlendWeight() const override;
-	virtual void ClearCachedBlendWeight() override;
-	virtual const FDeltaTimeRecord* GetDeltaTimeRecord() const override;
+	ANIMGRAPHRUNTIME_API virtual UAnimationAsset* GetAnimAsset() const override;
+	ANIMGRAPHRUNTIME_API virtual float GetAccumulatedTime() const override;
+	ANIMGRAPHRUNTIME_API virtual bool GetIgnoreForRelevancyTest() const override;
+	ANIMGRAPHRUNTIME_API virtual bool SetIgnoreForRelevancyTest(bool bInIgnoreForRelevancyTest) override;
+	ANIMGRAPHRUNTIME_API virtual float GetCachedBlendWeight() const override;
+	ANIMGRAPHRUNTIME_API virtual void ClearCachedBlendWeight() override;
+	ANIMGRAPHRUNTIME_API virtual const FDeltaTimeRecord* GetDeltaTimeRecord() const override;
 	// End of FAnimNode_RelevantAssetPlayerBase
 
 private:
diff --git a/Engine/Source/Runtime/AnimGraphRuntime/Public/AnimNodes/AnimNode_RefPose.h b/Engine/Source/Runtime/AnimGraphRuntime/Public/AnimNodes/AnimNode_RefPose.h
--- a/Engine/Source/Runtime/AnimGraphRuntime/Public/AnimNodes/AnimNode_RefPose.h
+++ b/Engine/Source/Runtime/AnimGraphRuntime/Public/AnimNodes/AnimNode_RefPose.h
@@ -16,7 +16,7 @@ enum ERefPoseType : int
 
 // RefPose pose nodes - ref pose or additive RefPose pose
 USTRUCT()
-struct ANIMGRAPHRUNTIME_API FAnimNode_RefPose : public FAnimNode_Base
+struct FAnimNode_RefPose : public FAnimNode_Base
 {
 	GENERATED_BODY()
 
@@ -35,18 +35,18 @@ public:
 #endif
 
 	// Get the type of this ref pose
-	ERefPoseType GetRefPoseType() const;
+	ANIMGRAPHRUNTIME_API ERefPoseType GetRefPoseType() const;
 	
-	virtual void Evaluate_AnyThread(FPoseContext& Output) override;
-	virtual void GatherDebugData(FNodeDebugData& DebugData) override;
+	ANIMGRAPHRUNTIME_API virtual void Evaluate_AnyThread(FPoseContext& Output) override;
+	ANIMGRAPHRUNTIME_API virtual void GatherDebugData(FNodeDebugData& DebugData) override;
 };
 
 USTRUCT()
-struct ANIMGRAPHRUNTIME_API FAnimNode_MeshSpaceRefPose : public FAnimNode_Base
+struct FAnimNode_MeshSpaceRefPose : public FAnimNode_Base
 {
 	GENERATED_BODY()
 public:	
 	FAnimNode_MeshSpaceRefPose() = default;
 
-	virtual void EvaluateComponentSpace_AnyThread(FComponentSpacePoseContext& Output);
+	ANIMGRAPHRUNTIME_API virtual void EvaluateComponentSpace_AnyThread(FComponentSpacePoseContext& Output);
 };
diff --git a/Engine/Source/Runtime/AnimGraphRuntime/Public/AnimNodes/AnimNode_RotateRootBone.h b/Engine/Source/Runtime/AnimGraphRuntime/Public/AnimNodes/AnimNode_RotateRootBone.h
--- a/Engine/Source/Runtime/AnimGraphRuntime/Public/AnimNodes/AnimNode_RotateRootBone.h
+++ b/Engine/Source/Runtime/AnimGraphRuntime/Public/AnimNodes/AnimNode_RotateRootBone.h
@@ -10,7 +10,7 @@
 
 //@TODO: Comment
 USTRUCT(BlueprintInternalUseOnly)
-struct ANIMGRAPHRUNTIME_API FAnimNode_RotateRootBone : public FAnimNode_Base
+struct FAnimNode_RotateRootBone : public FAnimNode_Base
 {
 	GENERATED_USTRUCT_BODY()
 
@@ -41,13 +41,13 @@ struct ANIMGRAPHRUNTIME_API FAnimNode_RotateRootBone : public FAnimNode_Base
 	float ActualYaw;
 
 public:	
-	FAnimNode_RotateRootBone();
+	ANIMGRAPHRUNTIME_API FAnimNode_RotateRootBone();
 
 	// FAnimNode_Base interface
-	virtual void Initialize_AnyThread(const FAnimationInitializeContext& Context) override;
-	virtual void CacheBones_AnyThread(const FAnimationCacheBonesContext& Context) override;
-	virtual void Update_AnyThread(const FAnimationUpdateContext& Context) override;
-	virtual void Evaluate_AnyThread(FPoseContext& Output) override;
-	virtual void GatherDebugData(FNodeDebugData& DebugData) override;
+	ANIMGRAPHRUNTIME_API virtual void Initialize_AnyThread(const FAnimationInitializeContext& Context) override;
+	ANIMGRAPHRUNTIME_API virtual void CacheBones_AnyThread(const FAnimationCacheBonesContext& Context) override;
+	ANIMGRAPHRUNTIME_API virtual void Update_AnyThread(const FAnimationUpdateContext& Context) override;
+	ANIMGRAPHRUNTIME_API virtual void Evaluate_AnyThread(FPoseContext& Output) override;
+	ANIMGRAPHRUNTIME_API virtual void GatherDebugData(FNodeDebugData& DebugData) override;
 	// End of FAnimNode_Base interface
 };
diff --git a/Engine/Source/Runtime/AnimGraphRuntime/Public/AnimNodes/AnimNode_RotationOffsetBlendSpace.h b/Engine/Source/Runtime/AnimGraphRuntime/Public/AnimNodes/AnimNode_RotationOffsetBlendSpace.h
--- a/Engine/Source/Runtime/AnimGraphRuntime/Public/AnimNodes/AnimNode_RotationOffsetBlendSpace.h
+++ b/Engine/Source/Runtime/AnimGraphRuntime/Public/AnimNodes/AnimNode_RotationOffsetBlendSpace.h
@@ -11,7 +11,7 @@
 
 //@TODO: Comment
 USTRUCT(BlueprintInternalUseOnly)
-struct ANIMGRAPHRUNTIME_API FAnimNode_RotationOffsetBlendSpace : public FAnimNode_BlendSpacePlayer
+struct FAnimNode_RotationOffsetBlendSpace : public FAnimNode_BlendSpacePlayer
 {
 	GENERATED_USTRUCT_BODY()
 
@@ -54,14 +54,14 @@ struct ANIMGRAPHRUNTIME_API FAnimNode_RotationOffsetBlendSpace : public FAnimNod
 	bool bIsLODEnabled;
 
 public:	
-	FAnimNode_RotationOffsetBlendSpace();
+	ANIMGRAPHRUNTIME_API FAnimNode_RotationOffsetBlendSpace();
 
 	// FAnimNode_Base interface
-	virtual void Initialize_AnyThread(const FAnimationInitializeContext& Context) override;
-	virtual void CacheBones_AnyThread(const FAnimationCacheBonesContext& Context) override;
-	virtual void UpdateAssetPlayer(const FAnimationUpdateContext& Context) override;
-	virtual void Evaluate_AnyThread(FPoseContext& Output) override;
-	virtual void GatherDebugData(FNodeDebugData& DebugData) override;
+	ANIMGRAPHRUNTIME_API virtual void Initialize_AnyThread(const FAnimationInitializeContext& Context) override;
+	ANIMGRAPHRUNTIME_API virtual void CacheBones_AnyThread(const FAnimationCacheBonesContext& Context) override;
+	ANIMGRAPHRUNTIME_API virtual void UpdateAssetPlayer(const FAnimationUpdateContext& Context) override;
+	ANIMGRAPHRUNTIME_API virtual void Evaluate_AnyThread(FPoseContext& Output) override;
+	ANIMGRAPHRUNTIME_API virtual void GatherDebugData(FNodeDebugData& DebugData) override;
 	virtual int32 GetLODThreshold() const override { return LODThreshold; }
 	// End of FAnimNode_Base interface
 };
diff --git a/Engine/Source/Runtime/AnimGraphRuntime/Public/AnimNodes/AnimNode_RotationOffsetBlendSpaceGraph.h b/Engine/Source/Runtime/AnimGraphRuntime/Public/AnimNodes/AnimNode_RotationOffsetBlendSpaceGraph.h
--- a/Engine/Source/Runtime/AnimGraphRuntime/Public/AnimNodes/AnimNode_RotationOffsetBlendSpaceGraph.h
+++ b/Engine/Source/Runtime/AnimGraphRuntime/Public/AnimNodes/AnimNode_RotationOffsetBlendSpaceGraph.h
@@ -9,7 +9,7 @@
 
 // Allows multiple animations to be blended between based on input parameters
 USTRUCT(BlueprintInternalUseOnly)
-struct ANIMGRAPHRUNTIME_API FAnimNode_RotationOffsetBlendSpaceGraph : public FAnimNode_BlendSpaceGraphBase
+struct FAnimNode_RotationOffsetBlendSpaceGraph : public FAnimNode_BlendSpaceGraphBase
 {
 	GENERATED_BODY()
 
@@ -61,11 +61,11 @@ private:
 
 private:
 	// FAnimNode_Base interface
-	virtual void Initialize_AnyThread(const FAnimationInitializeContext& Context) override;
-	virtual void CacheBones_AnyThread(const FAnimationCacheBonesContext& Context) override;
-	virtual void Update_AnyThread(const FAnimationUpdateContext& Context) override;
-	virtual void Evaluate_AnyThread(FPoseContext& Output) override;
-	virtual void GatherDebugData(FNodeDebugData& DebugData) override;
+	ANIMGRAPHRUNTIME_API virtual void Initialize_AnyThread(const FAnimationInitializeContext& Context) override;
+	ANIMGRAPHRUNTIME_API virtual void CacheBones_AnyThread(const FAnimationCacheBonesContext& Context) override;
+	ANIMGRAPHRUNTIME_API virtual void Update_AnyThread(const FAnimationUpdateContext& Context) override;
+	ANIMGRAPHRUNTIME_API virtual void Evaluate_AnyThread(FPoseContext& Output) override;
+	ANIMGRAPHRUNTIME_API virtual void GatherDebugData(FNodeDebugData& DebugData) override;
 	virtual int32 GetLODThreshold() const override { return LODThreshold; }
 	// End of FAnimNode_Base interface
 };
