[Engine]

diff --git a/Engine/Source/Runtime/AnimGraphRuntime/Public/AnimNodes/AnimNode_SequenceEvaluator.h b/Engine/Source/Runtime/AnimGraphRuntime/Public/AnimNodes/AnimNode_SequenceEvaluator.h
--- a/Engine/Source/Runtime/AnimGraphRuntime/Public/AnimNodes/AnimNode_SequenceEvaluator.h
+++ b/Engine/Source/Runtime/AnimGraphRuntime/Public/AnimNodes/AnimNode_SequenceEvaluator.h
@@ -26,7 +26,7 @@ namespace ESequenceEvalReinit
 // Typically the playback position of the animation for this node will represent something other than time, like jump height.
 // This node will not trigger any notifies present in the associated sequence.
 USTRUCT(BlueprintInternalUseOnly)
-struct ANIMGRAPHRUNTIME_API FAnimNode_SequenceEvaluatorBase : public FAnimNode_AssetPlayerBase
+struct FAnimNode_SequenceEvaluatorBase : public FAnimNode_AssetPlayerBase
 {
 	GENERATED_BODY()
 
@@ -35,16 +35,16 @@ private:
 
 public:
 	// FAnimNode_AssetPlayerBase interface
-	virtual float GetCurrentAssetTime() const override;
-	virtual float GetCurrentAssetLength() const override;
+	ANIMGRAPHRUNTIME_API virtual float GetCurrentAssetTime() const override;
+	ANIMGRAPHRUNTIME_API virtual float GetCurrentAssetLength() const override;
 	// End of FAnimNode_AssetPlayerBase interface
 
 	// FAnimNode_Base interface
-	virtual void Initialize_AnyThread(const FAnimationInitializeContext& Context) override;
-	virtual void CacheBones_AnyThread(const FAnimationCacheBonesContext& Context) override;
-	virtual void UpdateAssetPlayer(const FAnimationUpdateContext& Context) override;
-	virtual void Evaluate_AnyThread(FPoseContext& Output) override;
-	virtual void GatherDebugData(FNodeDebugData& DebugData) override;
+	ANIMGRAPHRUNTIME_API virtual void Initialize_AnyThread(const FAnimationInitializeContext& Context) override;
+	ANIMGRAPHRUNTIME_API virtual void CacheBones_AnyThread(const FAnimationCacheBonesContext& Context) override;
+	ANIMGRAPHRUNTIME_API virtual void UpdateAssetPlayer(const FAnimationUpdateContext& Context) override;
+	ANIMGRAPHRUNTIME_API virtual void Evaluate_AnyThread(FPoseContext& Output) override;
+	ANIMGRAPHRUNTIME_API virtual void GatherDebugData(FNodeDebugData& DebugData) override;
 	// End of FAnimNode_Base interface
 
 	// FAnimNode_AssetPlayerBase Interface
@@ -56,7 +56,7 @@ public:
 	void SetExplicitPreviousTime(float PreviousTime) { InternalTimeAccumulator = PreviousTime; }
 
 	// Get the effective delta time between the previous and current frame internal time
-	virtual float GetEffectiveDeltaTime(float ExplicitTime, float PrevExplicitTime) const;
+	ANIMGRAPHRUNTIME_API virtual float GetEffectiveDeltaTime(float ExplicitTime, float PrevExplicitTime) const;
 
 	// Set the animation sequence asset to evaluate
 	virtual bool SetSequence(UAnimSequenceBase* InSequence) { return false; }
@@ -97,7 +97,7 @@ public:
 
 // Sequence evaluator node that can be used with constant folding
 USTRUCT(BlueprintInternalUseOnly)
-struct ANIMGRAPHRUNTIME_API FAnimNode_SequenceEvaluator : public FAnimNode_SequenceEvaluatorBase
+struct FAnimNode_SequenceEvaluator : public FAnimNode_SequenceEvaluatorBase
 {
 	GENERATED_BODY()
 
@@ -150,30 +150,30 @@ private:
 
 public:
 	// FAnimNode_SequenceEvaluatorBase interface
-	virtual bool SetSequence(UAnimSequenceBase* InSequence) override;
-	virtual UAnimSequenceBase* GetSequence() const override;
-	virtual float GetExplicitTime() const override;
-	virtual bool SetExplicitTime(float InTime) override;
-	virtual bool SetShouldLoop(bool bInShouldLoop) override;
-	virtual bool GetTeleportToExplicitTime() const override;
-	virtual TEnumAsByte<ESequenceEvalReinit::Type> GetReinitializationBehavior() const override;
-	virtual float GetStartPosition() const override;
+	ANIMGRAPHRUNTIME_API virtual bool SetSequence(UAnimSequenceBase* InSequence) override;
+	ANIMGRAPHRUNTIME_API virtual UAnimSequenceBase* GetSequence() const override;
+	ANIMGRAPHRUNTIME_API virtual float GetExplicitTime() const override;
+	ANIMGRAPHRUNTIME_API virtual bool SetExplicitTime(float InTime) override;
+	ANIMGRAPHRUNTIME_API virtual bool SetShouldLoop(bool bInShouldLoop) override;
+	ANIMGRAPHRUNTIME_API virtual bool GetTeleportToExplicitTime() const override;
+	ANIMGRAPHRUNTIME_API virtual TEnumAsByte<ESequenceEvalReinit::Type> GetReinitializationBehavior() const override;
+	ANIMGRAPHRUNTIME_API virtual float GetStartPosition() const override;
 
 	// FAnimNode_AssetPlayerBase interface
-	virtual FName GetGroupName() const override;
-	virtual EAnimGroupRole::Type GetGroupRole() const override;
-	virtual EAnimSyncMethod GetGroupMethod() const override;
-	virtual bool GetIgnoreForRelevancyTest() const override;
-	virtual bool IsLooping() const override;
-	virtual bool SetGroupName(FName InGroupName) override;
-	virtual bool SetGroupRole(EAnimGroupRole::Type InRole) override;
-	virtual bool SetGroupMethod(EAnimSyncMethod InMethod) override;
-	virtual bool SetIgnoreForRelevancyTest(bool bInIgnoreForRelevancyTest) override;
+	ANIMGRAPHRUNTIME_API virtual FName GetGroupName() const override;
+	ANIMGRAPHRUNTIME_API virtual EAnimGroupRole::Type GetGroupRole() const override;
+	ANIMGRAPHRUNTIME_API virtual EAnimSyncMethod GetGroupMethod() const override;
+	ANIMGRAPHRUNTIME_API virtual bool GetIgnoreForRelevancyTest() const override;
+	ANIMGRAPHRUNTIME_API virtual bool IsLooping() const override;
+	ANIMGRAPHRUNTIME_API virtual bool SetGroupName(FName InGroupName) override;
+	ANIMGRAPHRUNTIME_API virtual bool SetGroupRole(EAnimGroupRole::Type InRole) override;
+	ANIMGRAPHRUNTIME_API virtual bool SetGroupMethod(EAnimSyncMethod InMethod) override;
+	ANIMGRAPHRUNTIME_API virtual bool SetIgnoreForRelevancyTest(bool bInIgnoreForRelevancyTest) override;
 };
 
 // Sequence evaluator node that can be used standalone (without constant folding)
 USTRUCT(BlueprintInternalUseOnly)
-struct ANIMGRAPHRUNTIME_API FAnimNode_SequenceEvaluator_Standalone : public FAnimNode_SequenceEvaluatorBase
+struct FAnimNode_SequenceEvaluator_Standalone : public FAnimNode_SequenceEvaluatorBase
 {
 	GENERATED_BODY()
 
@@ -242,4 +242,4 @@ public:
 	virtual bool SetGroupRole(EAnimGroupRole::Type InRole) override { GroupRole = InRole; return true; }
 	virtual bool SetGroupMethod(EAnimSyncMethod InMethod) override { Method = InMethod; return true; }
 	virtual bool SetIgnoreForRelevancyTest(bool bInIgnoreForRelevancyTest) override { bIgnoreForRelevancyTest = bInIgnoreForRelevancyTest; return true; }
-};
\ No newline at end of file
+};
diff --git a/Engine/Source/Runtime/AnimGraphRuntime/Public/AnimNodes/AnimNode_Slot.h b/Engine/Source/Runtime/AnimGraphRuntime/Public/AnimNodes/AnimNode_Slot.h
--- a/Engine/Source/Runtime/AnimGraphRuntime/Public/AnimNodes/AnimNode_Slot.h
+++ b/Engine/Source/Runtime/AnimGraphRuntime/Public/AnimNodes/AnimNode_Slot.h
@@ -12,7 +12,7 @@
 // game code can cause an animation to blend in and be played on the slot temporarily, overriding the
 // Source input.
 USTRUCT(BlueprintInternalUseOnly)
-struct ANIMGRAPHRUNTIME_API FAnimNode_Slot : public FAnimNode_Base
+struct FAnimNode_Slot : public FAnimNode_Base
 {
 	GENERATED_USTRUCT_BODY()
 
@@ -33,13 +33,13 @@ protected:
 	FGraphTraversalCounter SlotNodeInitializationCounter;
 
 public:	
-	FAnimNode_Slot();
+	ANIMGRAPHRUNTIME_API FAnimNode_Slot();
 
 	// FAnimNode_Base interface
-	virtual void Initialize_AnyThread(const FAnimationInitializeContext& Context) override;
-	virtual void CacheBones_AnyThread(const FAnimationCacheBonesContext& Context) override;
-	virtual void Update_AnyThread(const FAnimationUpdateContext& Context) override;
-	virtual void Evaluate_AnyThread(FPoseContext& Output) override;
-	virtual void GatherDebugData(FNodeDebugData& DebugData) override;
+	ANIMGRAPHRUNTIME_API virtual void Initialize_AnyThread(const FAnimationInitializeContext& Context) override;
+	ANIMGRAPHRUNTIME_API virtual void CacheBones_AnyThread(const FAnimationCacheBonesContext& Context) override;
+	ANIMGRAPHRUNTIME_API virtual void Update_AnyThread(const FAnimationUpdateContext& Context) override;
+	ANIMGRAPHRUNTIME_API virtual void Evaluate_AnyThread(FPoseContext& Output) override;
+	ANIMGRAPHRUNTIME_API virtual void GatherDebugData(FNodeDebugData& DebugData) override;
 	// End of FAnimNode_Base interface
 };
diff --git a/Engine/Source/Runtime/AnimGraphRuntime/Public/AnimNodes/AnimNode_Sync.h b/Engine/Source/Runtime/AnimGraphRuntime/Public/AnimNodes/AnimNode_Sync.h
--- a/Engine/Source/Runtime/AnimGraphRuntime/Public/AnimNodes/AnimNode_Sync.h
+++ b/Engine/Source/Runtime/AnimGraphRuntime/Public/AnimNodes/AnimNode_Sync.h
@@ -8,7 +8,7 @@
 #include "AnimNode_Sync.generated.h"
 
 USTRUCT(BlueprintInternalUseOnly)
-struct ANIMGRAPHRUNTIME_API FAnimNode_Sync : public FAnimNode_Base
+struct FAnimNode_Sync : public FAnimNode_Base
 {
 	GENERATED_BODY()
 
@@ -32,9 +32,9 @@ private:
 
 private:
 	// FAnimNode_Base
-	virtual void Initialize_AnyThread(const FAnimationInitializeContext& Context) override;
-	virtual void Update_AnyThread(const FAnimationUpdateContext& Context) override;
-	virtual void CacheBones_AnyThread(const FAnimationCacheBonesContext& Context) override;
-	virtual void Evaluate_AnyThread(FPoseContext& Output) override;
-	virtual void GatherDebugData(FNodeDebugData& DebugData) override;
-};
\ No newline at end of file
+	ANIMGRAPHRUNTIME_API virtual void Initialize_AnyThread(const FAnimationInitializeContext& Context) override;
+	ANIMGRAPHRUNTIME_API virtual void Update_AnyThread(const FAnimationUpdateContext& Context) override;
+	ANIMGRAPHRUNTIME_API virtual void CacheBones_AnyThread(const FAnimationCacheBonesContext& Context) override;
+	ANIMGRAPHRUNTIME_API virtual void Evaluate_AnyThread(FPoseContext& Output) override;
+	ANIMGRAPHRUNTIME_API virtual void GatherDebugData(FNodeDebugData& DebugData) override;
+};
diff --git a/Engine/Source/Runtime/AnimGraphRuntime/Public/AnimNodes/AnimNode_TwoWayBlend.h b/Engine/Source/Runtime/AnimGraphRuntime/Public/AnimNodes/AnimNode_TwoWayBlend.h
--- a/Engine/Source/Runtime/AnimGraphRuntime/Public/AnimNodes/AnimNode_TwoWayBlend.h
+++ b/Engine/Source/Runtime/AnimGraphRuntime/Public/AnimNodes/AnimNode_TwoWayBlend.h
@@ -10,7 +10,7 @@
 
 // This represents a baked transition
 USTRUCT(BlueprintInternalUseOnly)
-struct ANIMGRAPHRUNTIME_API FAnimNode_TwoWayBlend : public FAnimNode_Base
+struct FAnimNode_TwoWayBlend : public FAnimNode_Base
 {
 	GENERATED_USTRUCT_BODY()
 public:
@@ -75,11 +75,11 @@ public:
 	}
 
 	// FAnimNode_Base interface
-	virtual void Initialize_AnyThread(const FAnimationInitializeContext& Context) override;
-	virtual void CacheBones_AnyThread(const FAnimationCacheBonesContext& Context) override;
-	virtual void Update_AnyThread(const FAnimationUpdateContext& Context) override;
-	virtual void Evaluate_AnyThread(FPoseContext& Output) override;
-	virtual void GatherDebugData(FNodeDebugData& DebugData) override;
+	ANIMGRAPHRUNTIME_API virtual void Initialize_AnyThread(const FAnimationInitializeContext& Context) override;
+	ANIMGRAPHRUNTIME_API virtual void CacheBones_AnyThread(const FAnimationCacheBonesContext& Context) override;
+	ANIMGRAPHRUNTIME_API virtual void Update_AnyThread(const FAnimationUpdateContext& Context) override;
+	ANIMGRAPHRUNTIME_API virtual void Evaluate_AnyThread(FPoseContext& Output) override;
+	ANIMGRAPHRUNTIME_API virtual void GatherDebugData(FNodeDebugData& DebugData) override;
 	// End of FAnimNode_Base interface
 };
 
diff --git a/Engine/Source/Runtime/AnimGraphRuntime/Public/AnimNotifies/AnimNotify_PlayMontageNotify.h b/Engine/Source/Runtime/AnimGraphRuntime/Public/AnimNotifies/AnimNotify_PlayMontageNotify.h
--- a/Engine/Source/Runtime/AnimGraphRuntime/Public/AnimNotifies/AnimNotify_PlayMontageNotify.h
+++ b/Engine/Source/Runtime/AnimGraphRuntime/Public/AnimNotifies/AnimNotify_PlayMontageNotify.h
@@ -9,19 +9,19 @@
 // UAnimNotify_PlayMontageNotify
 //////////////////////////////////////////////////////////////////////////
 
-UCLASS(editinlinenew, const, hidecategories = Object, collapsecategories, meta = (DisplayName = "Montage Notify"))
-class ANIMGRAPHRUNTIME_API UAnimNotify_PlayMontageNotify : public UAnimNotify
+UCLASS(editinlinenew, const, hidecategories = Object, collapsecategories, meta = (DisplayName = "Montage Notify"), MinimalAPI)
+class UAnimNotify_PlayMontageNotify : public UAnimNotify
 {
 	GENERATED_UCLASS_BODY()
 
 public:
 
-	virtual void BranchingPointNotify(FBranchingPointNotifyPayload& BranchingPointPayload) override;
+	ANIMGRAPHRUNTIME_API virtual void BranchingPointNotify(FBranchingPointNotifyPayload& BranchingPointPayload) override;
 	
-	virtual FString GetNotifyName_Implementation() const override;
+	ANIMGRAPHRUNTIME_API virtual FString GetNotifyName_Implementation() const override;
 
 #if WITH_EDITOR
-	virtual bool CanBePlaced(UAnimSequenceBase* Animation) const override;
+	ANIMGRAPHRUNTIME_API virtual bool CanBePlaced(UAnimSequenceBase* Animation) const override;
 #endif
 protected:
 
@@ -35,19 +35,19 @@ protected:
 // UAnimNotify_PlayMontageNotifyWindow
 //////////////////////////////////////////////////////////////////////////
 
-UCLASS(editinlinenew, const, hidecategories = Object, collapsecategories, meta = (DisplayName = "Montage Notify Window"))
-class ANIMGRAPHRUNTIME_API UAnimNotify_PlayMontageNotifyWindow : public UAnimNotifyState
+UCLASS(editinlinenew, const, hidecategories = Object, collapsecategories, meta = (DisplayName = "Montage Notify Window"), MinimalAPI)
+class UAnimNotify_PlayMontageNotifyWindow : public UAnimNotifyState
 {
 	GENERATED_UCLASS_BODY()
 
 public:
-	virtual void BranchingPointNotifyBegin(FBranchingPointNotifyPayload& BranchingPointPayload) override;
-	virtual void BranchingPointNotifyEnd(FBranchingPointNotifyPayload& BranchingPointPayload) override;
+	ANIMGRAPHRUNTIME_API virtual void BranchingPointNotifyBegin(FBranchingPointNotifyPayload& BranchingPointPayload) override;
+	ANIMGRAPHRUNTIME_API virtual void BranchingPointNotifyEnd(FBranchingPointNotifyPayload& BranchingPointPayload) override;
 
-	virtual FString GetNotifyName_Implementation() const override;
+	ANIMGRAPHRUNTIME_API virtual FString GetNotifyName_Implementation() const override;
 
 #if WITH_EDITOR
-	virtual bool CanBePlaced(UAnimSequenceBase* Animation) const override;
+	ANIMGRAPHRUNTIME_API virtual bool CanBePlaced(UAnimSequenceBase* Animation) const override;
 #endif
 protected:
 
diff --git a/Engine/Source/Runtime/AnimGraphRuntime/Public/AnimSequencerInstance.h b/Engine/Source/Runtime/AnimGraphRuntime/Public/AnimSequencerInstance.h
--- a/Engine/Source/Runtime/AnimGraphRuntime/Public/AnimSequencerInstance.h
+++ b/Engine/Source/Runtime/AnimGraphRuntime/Public/AnimSequencerInstance.h
@@ -14,34 +14,34 @@
 struct FRootMotionOverride;
 struct FAnimSequencerData;
 
-UCLASS(transient, NotBlueprintable)
-class ANIMGRAPHRUNTIME_API UAnimSequencerInstance : public UAnimInstance, public ISequencerAnimationSupport
+UCLASS(transient, NotBlueprintable, MinimalAPI)
+class UAnimSequencerInstance : public UAnimInstance, public ISequencerAnimationSupport
 {
 	GENERATED_UCLASS_BODY()
 
 public:
 
 	/** Update an animation sequence player in this instance */
-	virtual void UpdateAnimTrack(UAnimSequenceBase* InAnimSequence, int32 SequenceId, float InPosition, float Weight, bool bFireNotifies);
-	virtual void UpdateAnimTrack(UAnimSequenceBase* InAnimSequence, int32 SequenceId, float InFromPosition, float InToPosition, float Weight, bool bFireNotifies);
+	ANIMGRAPHRUNTIME_API virtual void UpdateAnimTrack(UAnimSequenceBase* InAnimSequence, int32 SequenceId, float InPosition, float Weight, bool bFireNotifies);
+	ANIMGRAPHRUNTIME_API virtual void UpdateAnimTrack(UAnimSequenceBase* InAnimSequence, int32 SequenceId, float InFromPosition, float InToPosition, float Weight, bool bFireNotifies);
 
 	/** Update with Root Motion*/
 	UE_DEPRECATED(5.1, "Please use the UpdateAnimTrackWithRootMotion that takes FAnimSequencerData")
-	void UpdateAnimTrackWithRootMotion(UAnimSequenceBase* InAnimSequence, int32 SequenceId,const TOptional<FRootMotionOverride>& RootMotion, float InFromPosition, float InToPosition, float Weight, bool bFireNotifies, UMirrorDataTable* InMirrorDataTable);
+	ANIMGRAPHRUNTIME_API void UpdateAnimTrackWithRootMotion(UAnimSequenceBase* InAnimSequence, int32 SequenceId,const TOptional<FRootMotionOverride>& RootMotion, float InFromPosition, float InToPosition, float Weight, bool bFireNotifies, UMirrorDataTable* InMirrorDataTable);
 	
-	void UpdateAnimTrackWithRootMotion(const FAnimSequencerData& InAnimSequencerData);
+	ANIMGRAPHRUNTIME_API void UpdateAnimTrackWithRootMotion(const FAnimSequencerData& InAnimSequencerData);
 
 	/** Construct all nodes in this instance */
-	virtual void ConstructNodes() override;
+	ANIMGRAPHRUNTIME_API virtual void ConstructNodes() override;
 
 	/** Reset all nodes in this instance */
-	virtual void ResetNodes() override;
+	ANIMGRAPHRUNTIME_API virtual void ResetNodes() override;
 
 	/** Reset the pose for this instance*/
-	virtual void ResetPose() override;
+	ANIMGRAPHRUNTIME_API virtual void ResetPose() override;
 
 	/** Saved the named pose to restore after */
-	virtual void SavePose() override;
+	ANIMGRAPHRUNTIME_API virtual void SavePose() override;
 
 	virtual UAnimInstance* GetSourceAnimInstance() override { return this; }
 	virtual void SetSourceAnimInstance(UAnimInstance* SourceAnimInstance) {  /* nothing to do */ ensure(false); }
@@ -49,12 +49,12 @@ public:
 
 protected:
 	// UAnimInstance interface
-	virtual FAnimInstanceProxy* CreateAnimInstanceProxy() override;
+	ANIMGRAPHRUNTIME_API virtual FAnimInstanceProxy* CreateAnimInstanceProxy() override;
 
 
 
 public:
-	static const FName SequencerPoseName;
+	static ANIMGRAPHRUNTIME_API const FName SequencerPoseName;
 
 };
 
diff --git a/Engine/Source/Runtime/AnimGraphRuntime/Public/AnimSequencerInstanceProxy.h b/Engine/Source/Runtime/AnimGraphRuntime/Public/AnimSequencerInstanceProxy.h
--- a/Engine/Source/Runtime/AnimGraphRuntime/Public/AnimSequencerInstanceProxy.h
+++ b/Engine/Source/Runtime/AnimGraphRuntime/Public/AnimSequencerInstanceProxy.h
@@ -70,7 +70,7 @@ struct FRootMotionOverride
 	FTransform PreviousTransform;
 };
 
-struct ANIMGRAPHRUNTIME_API FAnimSequencerData
+struct FAnimSequencerData
 {
 	FAnimSequencerData(UAnimSequenceBase* InAnimSequence, int32 InSequenceId, const TOptional<FRootMotionOverride>& InRootMotion, float InFromPosition, float InToPosition, float InWeight, bool bInFireNotifies, ESwapRootBone InSwapRootBone, TOptional<FTransform> InInitialTransform, UMirrorDataTable* InMirrorDataTable)
 		: AnimSequence(InAnimSequence)
@@ -114,7 +114,7 @@ struct FSequencerPlayerAnimSequence : public FSequencerPlayerBase
 
 /** Proxy override for this UAnimInstance-derived class */
 USTRUCT()
-struct ANIMGRAPHRUNTIME_API FAnimSequencerInstanceProxy : public FAnimInstanceProxy
+struct FAnimSequencerInstanceProxy : public FAnimInstanceProxy
 {
 	GENERATED_BODY()
 
@@ -128,38 +128,38 @@ public:
 	{
 	}
 
-	virtual ~FAnimSequencerInstanceProxy();
+	ANIMGRAPHRUNTIME_API virtual ~FAnimSequencerInstanceProxy();
 
 	// FAnimInstanceProxy interface
-	virtual void Initialize(UAnimInstance* InAnimInstance) override;
-	virtual bool Evaluate(FPoseContext& Output) override;
-	virtual void PostEvaluate(UAnimInstance* InAnimInstance) override;
-	virtual void UpdateAnimationNode(const FAnimationUpdateContext& InContext) override;
+	ANIMGRAPHRUNTIME_API virtual void Initialize(UAnimInstance* InAnimInstance) override;
+	ANIMGRAPHRUNTIME_API virtual bool Evaluate(FPoseContext& Output) override;
+	ANIMGRAPHRUNTIME_API virtual void PostEvaluate(UAnimInstance* InAnimInstance) override;
+	ANIMGRAPHRUNTIME_API virtual void UpdateAnimationNode(const FAnimationUpdateContext& InContext) override;
 
 	/** Update an animation sequence player in this instance */
-	void UpdateAnimTrack(UAnimSequenceBase* InAnimSequence, uint32 SequenceId, float InPosition, float Weight, bool bFireNotifies);
-	void UpdateAnimTrack(UAnimSequenceBase* InAnimSequence, uint32 SequenceId, TOptional<float> InFromPosition, float InToPosition, float Weight, bool bFireNotifies);
+	ANIMGRAPHRUNTIME_API void UpdateAnimTrack(UAnimSequenceBase* InAnimSequence, uint32 SequenceId, float InPosition, float Weight, bool bFireNotifies);
+	ANIMGRAPHRUNTIME_API void UpdateAnimTrack(UAnimSequenceBase* InAnimSequence, uint32 SequenceId, TOptional<float> InFromPosition, float InToPosition, float Weight, bool bFireNotifies);
 	
 	UE_DEPRECATED(5.0, "Please use the UpdateAnimTrackWithRootMotion that takes a MirrorDataTable")
-	void UpdateAnimTrackWithRootMotion(UAnimSequenceBase* InAnimSequence, int32 SequenceId, const TOptional<FRootMotionOverride>& RootMotion, float InFromPosition, float InToPosition, float Weight, bool bFireNotifies);
+	ANIMGRAPHRUNTIME_API void UpdateAnimTrackWithRootMotion(UAnimSequenceBase* InAnimSequence, int32 SequenceId, const TOptional<FRootMotionOverride>& RootMotion, float InFromPosition, float InToPosition, float Weight, bool bFireNotifies);
 	
 	UE_DEPRECATED(5.1, "Please use the UpdateAnimTrackWithRootMotion that takes FAnimSequencerData")
-	void UpdateAnimTrackWithRootMotion(UAnimSequenceBase* InAnimSequence, int32 SequenceId, const TOptional<FRootMotionOverride>& RootMotion, float InFromPosition, float InToPosition, float Weight, bool bFireNotifies, UMirrorDataTable* InMirrorDataTable);
+	ANIMGRAPHRUNTIME_API void UpdateAnimTrackWithRootMotion(UAnimSequenceBase* InAnimSequence, int32 SequenceId, const TOptional<FRootMotionOverride>& RootMotion, float InFromPosition, float InToPosition, float Weight, bool bFireNotifies, UMirrorDataTable* InMirrorDataTable);
 
-	void UpdateAnimTrackWithRootMotion(const FAnimSequencerData& InAnimSequencerData);
+	ANIMGRAPHRUNTIME_API void UpdateAnimTrackWithRootMotion(const FAnimSequencerData& InAnimSequencerData);
 
 	/** Reset all nodes in this instance */
-	virtual void ResetNodes();
+	ANIMGRAPHRUNTIME_API virtual void ResetNodes();
 
 	/** Reset the pose in this instance*/
-	virtual void ResetPose();
+	ANIMGRAPHRUNTIME_API virtual void ResetPose();
 
 	/** Construct and link the base part of the blend tree */
-	virtual void ConstructNodes();
+	ANIMGRAPHRUNTIME_API virtual void ConstructNodes();
 
 protected:
 
-	void UpdateAnimTrack(UAnimSequenceBase* InAnimSequence, uint32 SequenceId, const TOptional<FRootMotionOverride>& RootMomtionOverride, TOptional<float> InFromPosition, float InToPosition, float Weight, bool bFireNotifies, UMirrorDataTable* InMirrorDataTable);
+	ANIMGRAPHRUNTIME_API void UpdateAnimTrack(UAnimSequenceBase* InAnimSequence, uint32 SequenceId, const TOptional<FRootMotionOverride>& RootMomtionOverride, TOptional<float> InFromPosition, float InToPosition, float Weight, bool bFireNotifies, UMirrorDataTable* InMirrorDataTable);
 
 	/** Find a player of a specified type */
 	template<typename Type>
@@ -189,9 +189,9 @@ protected:
 	/** custom root motion override sent in from sequencer */
 	TOptional<FRootMotionOverride> RootMotionOverride;
 
-	void InitAnimTrack(UAnimSequenceBase* InAnimSequence, uint32 SequenceId);
-	void EnsureAnimTrack(UAnimSequenceBase* InAnimSequence, uint32 SequenceId);
-	void ClearSequencePlayerAndMirrorMaps();
+	ANIMGRAPHRUNTIME_API void InitAnimTrack(UAnimSequenceBase* InAnimSequence, uint32 SequenceId);
+	ANIMGRAPHRUNTIME_API void EnsureAnimTrack(UAnimSequenceBase* InAnimSequence, uint32 SequenceId);
+	ANIMGRAPHRUNTIME_API void ClearSequencePlayerAndMirrorMaps();
 
 	ESwapRootBone SwapRootBone = ESwapRootBone::SwapRootBone_None;
 	TOptional<FTransform> InitialTransform;
diff --git a/Engine/Source/Runtime/AnimGraphRuntime/Public/AnimationStateMachineLibrary.h b/Engine/Source/Runtime/AnimGraphRuntime/Public/AnimationStateMachineLibrary.h
--- a/Engine/Source/Runtime/AnimGraphRuntime/Public/AnimationStateMachineLibrary.h
+++ b/Engine/Source/Runtime/AnimGraphRuntime/Public/AnimationStateMachineLibrary.h
@@ -30,15 +30,15 @@ struct FAnimationStateMachineReference : public FAnimNodeReference
 };
 
 // Exposes operations to be performed on anim state machine node contexts
-UCLASS(Experimental)
-class ANIMGRAPHRUNTIME_API UAnimationStateMachineLibrary : public UBlueprintFunctionLibrary
+UCLASS(Experimental, MinimalAPI)
+class UAnimationStateMachineLibrary : public UBlueprintFunctionLibrary
 {
 	GENERATED_BODY()
 
 public:
 	/** Get an anim state reference from an anim node reference */
 	UFUNCTION(BlueprintCallable, Category = "State Machine", meta=(BlueprintThreadSafe, DisplayName = "Convert to Animation State", ExpandEnumAsExecs = "Result"))
-	static void ConvertToAnimationStateResult(const FAnimNodeReference& Node, FAnimationStateResultReference& AnimationState, EAnimNodeReferenceConversionResult& Result);
+	static ANIMGRAPHRUNTIME_API void ConvertToAnimationStateResult(const FAnimNodeReference& Node, FAnimationStateResultReference& AnimationState, EAnimNodeReferenceConversionResult& Result);
 
 	/** Get an anim state reference from an anim node reference (pure) */
 	UFUNCTION(BlueprintPure, Category = "State Machine", meta=(BlueprintThreadSafe, DisplayName = "Convert to Animation State"))
@@ -51,7 +51,7 @@ public:
 
 	/** Get an anim state machine from an anim node reference */
 	UFUNCTION(BlueprintCallable, Category = "State Machine", meta=(BlueprintThreadSafe, DisplayName = "Convert to Animation State Machine", ExpandEnumAsExecs = "Result"))
-	static void ConvertToAnimationStateMachine(const FAnimNodeReference& Node, FAnimationStateMachineReference& AnimationState, EAnimNodeReferenceConversionResult& Result);
+	static ANIMGRAPHRUNTIME_API void ConvertToAnimationStateMachine(const FAnimNodeReference& Node, FAnimationStateMachineReference& AnimationState, EAnimNodeReferenceConversionResult& Result);
 
 
 	/** Get an anim state machine from an anim node reference (pure) */
@@ -65,27 +65,27 @@ public:
 	
 	/** Returns whether the state the node belongs to is blending in */
 	UFUNCTION(BlueprintPure, Category = "State Machine", meta=(BlueprintThreadSafe))
-	static bool IsStateBlendingIn(const FAnimUpdateContext& UpdateContext, const FAnimationStateResultReference& Node);
+	static ANIMGRAPHRUNTIME_API bool IsStateBlendingIn(const FAnimUpdateContext& UpdateContext, const FAnimationStateResultReference& Node);
 
 	/** Returns whether the state the node belongs to is blending out */
 	UFUNCTION(BlueprintPure, Category = "State Machine", meta=(BlueprintThreadSafe))
-	static bool IsStateBlendingOut(const FAnimUpdateContext& UpdateContext, const FAnimationStateResultReference& Node);
+	static ANIMGRAPHRUNTIME_API bool IsStateBlendingOut(const FAnimUpdateContext& UpdateContext, const FAnimationStateResultReference& Node);
 
 	/** Manually set the current state of the state machine
 		NOTE: Custom blend type is not supported */
 	UFUNCTION(BlueprintCallable, Category = "State Machine", meta=(BlueprintThreadSafe, AdvancedDisplay = "4"))
-	static void SetState(const FAnimUpdateContext& UpdateContext, const FAnimationStateMachineReference& Node, FName TargetState, float Duration
+	static ANIMGRAPHRUNTIME_API void SetState(const FAnimUpdateContext& UpdateContext, const FAnimationStateMachineReference& Node, FName TargetState, float Duration
 		, TEnumAsByte<ETransitionLogicType::Type> BlendType, UBlendProfile* BlendProfile, EAlphaBlendOption AlphaBlendOption, UCurveFloat* CustomBlendCurve);
 
 	/** Returns the name of the current state of this state machine */
 	UFUNCTION(BlueprintPure, Category = "State Machine", meta=(BlueprintThreadSafe))
-	static FName GetState(const FAnimUpdateContext& UpdateContext, const FAnimationStateMachineReference& Node);
+	static ANIMGRAPHRUNTIME_API FName GetState(const FAnimUpdateContext& UpdateContext, const FAnimationStateMachineReference& Node);
 
 	/** Returns the remaining animation time of the state's most relevant asset player */
 	UFUNCTION(BlueprintPure, Category = "State Machine", meta = (BlueprintThreadSafe))
-	static float GetRelevantAnimTimeRemaining(const FAnimUpdateContext& UpdateContext, const FAnimationStateResultReference& Node);
+	static ANIMGRAPHRUNTIME_API float GetRelevantAnimTimeRemaining(const FAnimUpdateContext& UpdateContext, const FAnimationStateResultReference& Node);
 
 	/** Returns the remaining animation time as a fraction of the duration for the state's most relevant asset player */
 	UFUNCTION(BlueprintPure, Category = "State Machine", meta = (BlueprintThreadSafe))
-	static float GetRelevantAnimTimeRemainingFraction(const FAnimUpdateContext& UpdateContext, const FAnimationStateResultReference& Node);
-};
\ No newline at end of file
+	static ANIMGRAPHRUNTIME_API float GetRelevantAnimTimeRemainingFraction(const FAnimUpdateContext& UpdateContext, const FAnimationStateResultReference& Node);
+};
diff --git a/Engine/Source/Runtime/AnimGraphRuntime/Public/BlendListBaseLibrary.h b/Engine/Source/Runtime/AnimGraphRuntime/Public/BlendListBaseLibrary.h
--- a/Engine/Source/Runtime/AnimGraphRuntime/Public/BlendListBaseLibrary.h
+++ b/Engine/Source/Runtime/AnimGraphRuntime/Public/BlendListBaseLibrary.h
@@ -18,8 +18,8 @@ struct FBlendListBaseReference : public FAnimNodeReference
 };
 
 // Exposes operations to be performed on anim state machine node contexts
-UCLASS(Experimental)
-class ANIMGRAPHRUNTIME_API UBlendListBaseLibrary : public UBlueprintFunctionLibrary
+UCLASS(Experimental, MinimalAPI)
+class UBlendListBaseLibrary : public UBlueprintFunctionLibrary
 {
 	GENERATED_BODY()
 
@@ -27,11 +27,11 @@ public:
 
 	/** Get a blend list base context from an anim node context. */
 	UFUNCTION(BlueprintCallable, Category = "Blend List Base", meta = (BlueprintThreadSafe, ExpandEnumAsExecs = "Result"))
-	static FBlendListBaseReference ConvertToBlendListBase(const FAnimNodeReference& Node, EAnimNodeReferenceConversionResult& Result);
+	static ANIMGRAPHRUNTIME_API FBlendListBaseReference ConvertToBlendListBase(const FAnimNodeReference& Node, EAnimNodeReferenceConversionResult& Result);
 
 	/** Reset target blend list node to that the next blend is executed from a blank state */
 	UFUNCTION(BlueprintCallable, Category = "Blend List Base", meta = (BlueprintThreadSafe))
-	static void ResetNode(const FBlendListBaseReference& BlendListBase);
+	static ANIMGRAPHRUNTIME_API void ResetNode(const FBlendListBaseReference& BlendListBase);
 
 
 };
diff --git a/Engine/Source/Runtime/AnimGraphRuntime/Public/BlendSpaceLibrary.h b/Engine/Source/Runtime/AnimGraphRuntime/Public/BlendSpaceLibrary.h
--- a/Engine/Source/Runtime/AnimGraphRuntime/Public/BlendSpaceLibrary.h
+++ b/Engine/Source/Runtime/AnimGraphRuntime/Public/BlendSpaceLibrary.h
@@ -21,15 +21,15 @@ struct FBlendSpaceReference : public FAnimNodeReference
 /**
  * Exposes operations to be performed on a blend space anim node.
  */
-UCLASS()
-class ANIMGRAPHRUNTIME_API UBlendSpaceLibrary : public UBlueprintFunctionLibrary
+UCLASS(MinimalAPI)
+class UBlendSpaceLibrary : public UBlueprintFunctionLibrary
 {
 	GENERATED_BODY()
 
 public:
 	/** Get a blend space context from an anim node context. */
 	UFUNCTION(BlueprintCallable, Category = "Blend Space", meta = (BlueprintThreadSafe, ExpandEnumAsExecs = "Result"))
-	static FBlendSpaceReference ConvertToBlendSpace(const FAnimNodeReference& Node, EAnimNodeReferenceConversionResult& Result);
+	static ANIMGRAPHRUNTIME_API FBlendSpaceReference ConvertToBlendSpace(const FAnimNodeReference& Node, EAnimNodeReferenceConversionResult& Result);
 
 	/** Get a blend space context from an anim node context (pure). */
 	UFUNCTION(BlueprintPure, Category = "Blend Space", meta = (BlueprintThreadSafe, DisplayName = "Convert to Blend Space (Pure)"))
@@ -42,15 +42,15 @@ public:
 
 	/** Get the current position of the blend space. */
 	UFUNCTION(BlueprintPure, Category = "Blend Space", meta = (BlueprintThreadSafe))
-	static FVector GetPosition(const FBlendSpaceReference& BlendSpace);
+	static ANIMGRAPHRUNTIME_API FVector GetPosition(const FBlendSpaceReference& BlendSpace);
 
 	/** Get the current sample coordinates after going through the filtering. */
 	UFUNCTION(BlueprintPure, Category = "Blend Space", meta = (BlueprintThreadSafe))
-	static FVector GetFilteredPosition(const FBlendSpaceReference& BlendSpace);
+	static ANIMGRAPHRUNTIME_API FVector GetFilteredPosition(const FBlendSpaceReference& BlendSpace);
 
 	/** Forces the Position to the specified value */
 	UFUNCTION(BlueprintCallable, Category = "Blend Space", meta = (BlueprintThreadSafe))
-	static void SnapToPosition(const FBlendSpaceReference& BlendSpace, FVector NewPosition);
+	static ANIMGRAPHRUNTIME_API void SnapToPosition(const FBlendSpaceReference& BlendSpace, FVector NewPosition);
 };
 
 
diff --git a/Engine/Source/Runtime/AnimGraphRuntime/Public/BlendSpacePlayerLibrary.h b/Engine/Source/Runtime/AnimGraphRuntime/Public/BlendSpacePlayerLibrary.h
--- a/Engine/Source/Runtime/AnimGraphRuntime/Public/BlendSpacePlayerLibrary.h
+++ b/Engine/Source/Runtime/AnimGraphRuntime/Public/BlendSpacePlayerLibrary.h
@@ -21,15 +21,15 @@ struct FBlendSpacePlayerReference : public FAnimNodeReference
 /**
  * Exposes operations to be performed on a blend space player anim node.
  */
-UCLASS()
-class ANIMGRAPHRUNTIME_API UBlendSpacePlayerLibrary : public UBlueprintFunctionLibrary
+UCLASS(MinimalAPI)
+class UBlendSpacePlayerLibrary : public UBlueprintFunctionLibrary
 {
 	GENERATED_BODY()
 
 public:
 	/** Get a blend space player context from an anim node context. */
 	UFUNCTION(BlueprintCallable, Category = "Blend Space Player", meta = (BlueprintThreadSafe, ExpandEnumAsExecs = "Result"))
-	static FBlendSpacePlayerReference ConvertToBlendSpacePlayer(const FAnimNodeReference& Node, EAnimNodeReferenceConversionResult& Result);
+	static ANIMGRAPHRUNTIME_API FBlendSpacePlayerReference ConvertToBlendSpacePlayer(const FAnimNodeReference& Node, EAnimNodeReferenceConversionResult& Result);
 
 	/** Get a blend space player context from an anim node context (pure). */
 	UFUNCTION(BlueprintPure, Category = "Blend Space Player", meta = (BlueprintThreadSafe, DisplayName = "Convert to Blend Space Player"))
@@ -42,51 +42,51 @@ public:
 
 	/** Set the current BlendSpace of the blend space player. */
 	UFUNCTION(BlueprintCallable, Category = "Blend Space Player", meta = (BlueprintThreadSafe))
-	static FBlendSpacePlayerReference SetBlendSpace(const FBlendSpacePlayerReference& BlendSpacePlayer, UBlendSpace* BlendSpace);
+	static ANIMGRAPHRUNTIME_API FBlendSpacePlayerReference SetBlendSpace(const FBlendSpacePlayerReference& BlendSpacePlayer, UBlendSpace* BlendSpace);
 
 	/** Set the current BlendSpace of the blend space player with an interial blend time. */
 	UFUNCTION(BlueprintCallable, Category = "Blend Space Player", meta = (BlueprintThreadSafe))
-	static FBlendSpacePlayerReference SetBlendSpaceWithInertialBlending(const FAnimUpdateContext& UpdateContext, const FBlendSpacePlayerReference& BlendSpacePlayer, UBlendSpace* BlendSpace, float BlendTime = 0.2f);
+	static ANIMGRAPHRUNTIME_API FBlendSpacePlayerReference SetBlendSpaceWithInertialBlending(const FAnimUpdateContext& UpdateContext, const FBlendSpacePlayerReference& BlendSpacePlayer, UBlendSpace* BlendSpace, float BlendTime = 0.2f);
 
 	/** Set whether the current play time should reset when BlendSpace changes of the blend space player. */
 	UFUNCTION(BlueprintCallable, Category = "Blend Space Player", meta = (BlueprintThreadSafe))
-	static FBlendSpacePlayerReference SetResetPlayTimeWhenBlendSpaceChanges(const FBlendSpacePlayerReference& BlendSpacePlayer, bool bReset);
+	static ANIMGRAPHRUNTIME_API FBlendSpacePlayerReference SetResetPlayTimeWhenBlendSpaceChanges(const FBlendSpacePlayerReference& BlendSpacePlayer, bool bReset);
 	
 	/** Set the play rate of the blend space player. */
 	UFUNCTION(BlueprintCallable, Category = "Blend Space Player", meta = (BlueprintThreadSafe))
-	static FBlendSpacePlayerReference SetPlayRate(const FBlendSpacePlayerReference& BlendSpacePlayer, float PlayRate);
+	static ANIMGRAPHRUNTIME_API FBlendSpacePlayerReference SetPlayRate(const FBlendSpacePlayerReference& BlendSpacePlayer, float PlayRate);
 
 	/** Set the loop of the blend space player. */
 	UFUNCTION(BlueprintCallable, Category = "Blend Space Player", meta = (BlueprintThreadSafe))
-	static FBlendSpacePlayerReference SetLoop(const FBlendSpacePlayerReference& BlendSpacePlayer, bool bLoop);
+	static ANIMGRAPHRUNTIME_API FBlendSpacePlayerReference SetLoop(const FBlendSpacePlayerReference& BlendSpacePlayer, bool bLoop);
 
 	/** Get the current BlendSpace of the blend space player. */
 	UFUNCTION(BlueprintPure, Category = "Blend Space Player", meta = (BlueprintThreadSafe))
-	static UBlendSpace* GetBlendSpace(const FBlendSpacePlayerReference& BlendSpacePlayer);
+	static ANIMGRAPHRUNTIME_API UBlendSpace* GetBlendSpace(const FBlendSpacePlayerReference& BlendSpacePlayer);
 
 	/** Get the current position of the blend space player. */
 	UFUNCTION(BlueprintPure, Category = "Blend Space Player", meta = (BlueprintThreadSafe))
-	static FVector GetPosition(const FBlendSpacePlayerReference& BlendSpacePlayer);
+	static ANIMGRAPHRUNTIME_API FVector GetPosition(const FBlendSpacePlayerReference& BlendSpacePlayer);
 
 	/** Get the current start position of the blend space player. */
 	UFUNCTION(BlueprintPure, Category = "Blend Space Player", meta = (BlueprintThreadSafe))
-	static float GetStartPosition(const FBlendSpacePlayerReference& BlendSpacePlayer);
+	static ANIMGRAPHRUNTIME_API float GetStartPosition(const FBlendSpacePlayerReference& BlendSpacePlayer);
 
 	/** Get the current play rate of the blend space player. */
 	UFUNCTION(BlueprintPure, Category = "Blend Space Player", meta = (BlueprintThreadSafe))
-	static float GetPlayRate(const FBlendSpacePlayerReference& BlendSpacePlayer);
+	static ANIMGRAPHRUNTIME_API float GetPlayRate(const FBlendSpacePlayerReference& BlendSpacePlayer);
 
 	/** Get the current loop of the blend space player.  */
 	UFUNCTION(BlueprintPure, Category = "Blend Space Player", meta = (BlueprintThreadSafe))
-	static bool GetLoop(const FBlendSpacePlayerReference& BlendSpacePlayer);
+	static ANIMGRAPHRUNTIME_API bool GetLoop(const FBlendSpacePlayerReference& BlendSpacePlayer);
 
 	/** Get the current value of whether the current play time should reset when BlendSpace changes of the blend space player. */
 	UFUNCTION(BlueprintPure, Category = "Blend Space Player", meta = (BlueprintThreadSafe))
-	static bool ShouldResetPlayTimeWhenBlendSpaceChanges(const FBlendSpacePlayerReference& BlendSpacePlayer); 
+	static ANIMGRAPHRUNTIME_API bool ShouldResetPlayTimeWhenBlendSpaceChanges(const FBlendSpacePlayerReference& BlendSpacePlayer); 
 
 	/** Forces the Position to the specified value */
 	UFUNCTION(BlueprintCallable, Category = "Blend Space Player", meta = (BlueprintThreadSafe))
-	static void SnapToPosition(const FBlendSpacePlayerReference& BlendSpacePlayer, FVector NewPosition);
+	static ANIMGRAPHRUNTIME_API void SnapToPosition(const FBlendSpacePlayerReference& BlendSpacePlayer, FVector NewPosition);
 };
 
 
diff --git a/Engine/Source/Runtime/AnimGraphRuntime/Public/BoneControllers/AnimNode_AnimDynamics.h b/Engine/Source/Runtime/AnimGraphRuntime/Public/BoneControllers/AnimNode_AnimDynamics.h
--- a/Engine/Source/Runtime/AnimGraphRuntime/Public/BoneControllers/AnimNode_AnimDynamics.h
+++ b/Engine/Source/Runtime/AnimGraphRuntime/Public/BoneControllers/AnimNode_AnimDynamics.h
@@ -283,11 +283,11 @@ struct FAnimConstraintOffsetPair
 
 
 USTRUCT(BlueprintInternalUseOnly)
-struct ANIMGRAPHRUNTIME_API FAnimNode_AnimDynamics : public FAnimNode_SkeletalControlBase
+struct FAnimNode_AnimDynamics : public FAnimNode_SkeletalControlBase
 {
 	GENERATED_BODY();
 
-	FAnimNode_AnimDynamics();
+	ANIMGRAPHRUNTIME_API FAnimNode_AnimDynamics();
 
 	/**
 	* Overridden linear damping value. The default is 0.7. Values below 0.7 won't have an effect.
@@ -474,69 +474,69 @@ struct ANIMGRAPHRUNTIME_API FAnimNode_AnimDynamics : public FAnimNode_SkeletalCo
 #endif
 
 	// FAnimNode_SkeletalControlBase interface
-	virtual void Initialize_AnyThread(const FAnimationInitializeContext& Context) override;
-	virtual void UpdateInternal(const FAnimationUpdateContext& Context) override;
-	virtual void EvaluateSkeletalControl_AnyThread(FComponentSpacePoseContext& Output, TArray<FBoneTransform>& OutBoneTransforms) override;
-	virtual void GatherDebugData(FNodeDebugData& DebugData) override;
-	virtual bool HasPreUpdate() const override;
-	virtual void PreUpdate(const UAnimInstance* InAnimInstance) override;
+	ANIMGRAPHRUNTIME_API virtual void Initialize_AnyThread(const FAnimationInitializeContext& Context) override;
+	ANIMGRAPHRUNTIME_API virtual void UpdateInternal(const FAnimationUpdateContext& Context) override;
+	ANIMGRAPHRUNTIME_API virtual void EvaluateSkeletalControl_AnyThread(FComponentSpacePoseContext& Output, TArray<FBoneTransform>& OutBoneTransforms) override;
+	ANIMGRAPHRUNTIME_API virtual void GatherDebugData(FNodeDebugData& DebugData) override;
+	ANIMGRAPHRUNTIME_API virtual bool HasPreUpdate() const override;
+	ANIMGRAPHRUNTIME_API virtual void PreUpdate(const UAnimInstance* InAnimInstance) override;
 	virtual bool NeedsDynamicReset() const override { return true; }
 	virtual void ResetDynamics(ETeleportType InTeleportType) override { RequestInitialise(InTeleportType); }
-	virtual int32 GetLODThreshold() const override;
+	ANIMGRAPHRUNTIME_API virtual int32 GetLODThreshold() const override;
 	// End of FAnimNode_SkeletalControlBase interface
 
-	void RequestInitialise(ETeleportType InTeleportType);
-	void InitPhysics(FComponentSpacePoseContext& Output);
-	void TermPhysics();
+	ANIMGRAPHRUNTIME_API void RequestInitialise(ETeleportType InTeleportType);
+	ANIMGRAPHRUNTIME_API void InitPhysics(FComponentSpacePoseContext& Output);
+	ANIMGRAPHRUNTIME_API void TermPhysics();
 
-	void UpdateChainPhysicsBodyDefinitions(const FReferenceSkeleton& ReferenceSkeleton);
-	void ValidateChainPhysicsBodyDefinitions(const FReferenceSkeleton& ReferenceSkeleton);
-	void FindChainBoneNames(const FReferenceSkeleton& ReferenceSkeleton, TArray<FName>& ChainBoneNames);
-	void UpdateLimits(FComponentSpacePoseContext& Output);
+	ANIMGRAPHRUNTIME_API void UpdateChainPhysicsBodyDefinitions(const FReferenceSkeleton& ReferenceSkeleton);
+	ANIMGRAPHRUNTIME_API void ValidateChainPhysicsBodyDefinitions(const FReferenceSkeleton& ReferenceSkeleton);
+	ANIMGRAPHRUNTIME_API void FindChainBoneNames(const FReferenceSkeleton& ReferenceSkeleton, TArray<FName>& ChainBoneNames);
+	ANIMGRAPHRUNTIME_API void UpdateLimits(FComponentSpacePoseContext& Output);
 
-	int32 GetNumBodies() const;
-	const FAnimPhysRigidBody& GetPhysBody(int32 BodyIndex) const;
+	ANIMGRAPHRUNTIME_API int32 GetNumBodies() const;
+	ANIMGRAPHRUNTIME_API const FAnimPhysRigidBody& GetPhysBody(int32 BodyIndex) const;
 
-	FTransform GetBodyComponentSpaceTransform(const FAnimPhysRigidBody& Body, const USkeletalMeshComponent* const SkelComp) const;
+	ANIMGRAPHRUNTIME_API FTransform GetBodyComponentSpaceTransform(const FAnimPhysRigidBody& Body, const USkeletalMeshComponent* const SkelComp) const;
 
 #if WITH_EDITOR
 
 	// Accessors for editor code (mainly for visualization functions)
-	FVector GetBodyLocalJointOffset(const int32 BodyIndex) const;
+	ANIMGRAPHRUNTIME_API FVector GetBodyLocalJointOffset(const int32 BodyIndex) const;
 
 	// True by default, if false physics simulation will not update this frame. Used to prevent the rig moving whilst interactively editing parameters with a widget in the viewport.
 	bool bDoPhysicsUpdateInEditor;
 
 #endif
 
-	bool ShouldDoPhysicsUpdate() const;
+	ANIMGRAPHRUNTIME_API bool ShouldDoPhysicsUpdate() const;
 
 protected:
 
 	// FAnimNode_SkeletalControlBase protected interface
-	virtual void InitializeBoneReferences(const FBoneContainer& RequiredBones) override;
-	virtual bool IsValidToEvaluate(const USkeleton* Skeleton, const FBoneContainer& RequiredBones);
+	ANIMGRAPHRUNTIME_API virtual void InitializeBoneReferences(const FBoneContainer& RequiredBones) override;
+	ANIMGRAPHRUNTIME_API virtual bool IsValidToEvaluate(const USkeleton* Skeleton, const FBoneContainer& RequiredBones);
 	// End of FAnimNode_SkeletalControlBase protected interface
 
 private:
 	// Given a bone index, get it's transform in the currently selected simulation space
-	FTransform GetBoneTransformInSimSpace(FComponentSpacePoseContext& Output, const FCompactPoseBoneIndex& BoneIndex) const;
+	ANIMGRAPHRUNTIME_API FTransform GetBoneTransformInSimSpace(FComponentSpacePoseContext& Output, const FCompactPoseBoneIndex& BoneIndex) const;
 
 	// Given a transform in simulation space, convert it back to component space
-	FTransform GetComponentSpaceTransformFromSimSpace(AnimPhysSimSpaceType SimSpace, FComponentSpacePoseContext& Output, const FTransform& InSimTransform) const;
-	FTransform GetComponentSpaceTransformFromSimSpace(AnimPhysSimSpaceType SimSpace, FComponentSpacePoseContext& Output, const FTransform& InSimTransform, const FTransform& InCompWorldSpaceTM, const FTransform& InActorWorldSpaceTM) const;
-	FTransform GetComponentSpaceTransformFromSimSpace(AnimPhysSimSpaceType SimSpace, const USkeletalMeshComponent* const SkelComp, const FTransform& InSimTransform) const;
+	ANIMGRAPHRUNTIME_API FTransform GetComponentSpaceTransformFromSimSpace(AnimPhysSimSpaceType SimSpace, FComponentSpacePoseContext& Output, const FTransform& InSimTransform) const;
+	ANIMGRAPHRUNTIME_API FTransform GetComponentSpaceTransformFromSimSpace(AnimPhysSimSpaceType SimSpace, FComponentSpacePoseContext& Output, const FTransform& InSimTransform, const FTransform& InCompWorldSpaceTM, const FTransform& InActorWorldSpaceTM) const;
+	ANIMGRAPHRUNTIME_API FTransform GetComponentSpaceTransformFromSimSpace(AnimPhysSimSpaceType SimSpace, const USkeletalMeshComponent* const SkelComp, const FTransform& InSimTransform) const;
 
 	// Given a transform in component space, convert it to the current sim space
-	FTransform GetSimSpaceTransformFromComponentSpace(AnimPhysSimSpaceType SimSpace, FComponentSpacePoseContext& Output, const FTransform& InComponentTransform) const;
+	ANIMGRAPHRUNTIME_API FTransform GetSimSpaceTransformFromComponentSpace(AnimPhysSimSpaceType SimSpace, FComponentSpacePoseContext& Output, const FTransform& InComponentTransform) const;
 
 	// Given a world-space vector, convert it into the current simulation space
-	FVector TransformWorldVectorToSimSpace(FComponentSpacePoseContext& Output, const FVector& InVec) const;
+	ANIMGRAPHRUNTIME_API FVector TransformWorldVectorToSimSpace(FComponentSpacePoseContext& Output, const FVector& InVec) const;
 
-	void ConvertSimulationSpace(FComponentSpacePoseContext& Output, AnimPhysSimSpaceType From, AnimPhysSimSpaceType To) const;
+	ANIMGRAPHRUNTIME_API void ConvertSimulationSpace(FComponentSpacePoseContext& Output, AnimPhysSimSpaceType From, AnimPhysSimSpaceType To) const;
 
 	// Maximum time to consider when accumulating time debt to avoid spiraling
-	static const float MaxTimeDebt;
+	static ANIMGRAPHRUNTIME_API const float MaxTimeDebt;
 
 	// Cached timestep from the update phase (needed in evaluate phase)
 	float NextTimeStep;
@@ -590,10 +590,10 @@ private:
 	// Live debug
 	//////////////////////////////////////////////////////////////////////////
 #if ENABLE_ANIM_DRAW_DEBUG
-	void DrawBodies(FComponentSpacePoseContext& InContext, const TArray<FAnimPhysRigidBody*>& InBodies);
+	ANIMGRAPHRUNTIME_API void DrawBodies(FComponentSpacePoseContext& InContext, const TArray<FAnimPhysRigidBody*>& InBodies);
 
 	int32 FilteredBoneIndex;
 #endif
 public: 
-	static bool IsAnimDynamicsSystemEnabledFor(int32 InLOD);
+	static ANIMGRAPHRUNTIME_API bool IsAnimDynamicsSystemEnabledFor(int32 InLOD);
 };
diff --git a/Engine/Source/Runtime/AnimGraphRuntime/Public/BoneControllers/AnimNode_ApplyLimits.h b/Engine/Source/Runtime/AnimGraphRuntime/Public/BoneControllers/AnimNode_ApplyLimits.h
--- a/Engine/Source/Runtime/AnimGraphRuntime/Public/BoneControllers/AnimNode_ApplyLimits.h
+++ b/Engine/Source/Runtime/AnimGraphRuntime/Public/BoneControllers/AnimNode_ApplyLimits.h
@@ -32,11 +32,11 @@ struct FAngularRangeLimit
 };
 
 USTRUCT(BlueprintInternalUseOnly)
-struct ANIMGRAPHRUNTIME_API FAnimNode_ApplyLimits : public FAnimNode_SkeletalControlBase
+struct FAnimNode_ApplyLimits : public FAnimNode_SkeletalControlBase
 {
 	GENERATED_USTRUCT_BODY();
 
-	FAnimNode_ApplyLimits();
+	ANIMGRAPHRUNTIME_API FAnimNode_ApplyLimits();
 
 	UPROPERTY(EditAnywhere, Category = Angular)
 	TArray<FAngularRangeLimit> AngularRangeLimits;
@@ -44,21 +44,21 @@ struct ANIMGRAPHRUNTIME_API FAnimNode_ApplyLimits : public FAnimNode_SkeletalCon
 	UPROPERTY(EditAnywhere, BlueprintReadWrite, EditFixedSize, Category = "Angular", meta = (PinHiddenByDefault))
 	TArray<FVector> AngularOffsets;
 
-	void RecalcLimits();
+	ANIMGRAPHRUNTIME_API void RecalcLimits();
 
 	// FAnimNode_Base interface
-	virtual void GatherDebugData(FNodeDebugData& DebugData) override;
-	virtual void OnInitializeAnimInstance(const FAnimInstanceProxy* InProxy, const UAnimInstance* InAnimInstance) override;
+	ANIMGRAPHRUNTIME_API virtual void GatherDebugData(FNodeDebugData& DebugData) override;
+	ANIMGRAPHRUNTIME_API virtual void OnInitializeAnimInstance(const FAnimInstanceProxy* InProxy, const UAnimInstance* InAnimInstance) override;
 	virtual bool NeedsOnInitializeAnimInstance() const override { return true; }
 	// End of FAnimNode_Base interface
 
 	// FAnimNode_SkeletalControlBase interface
-	virtual void EvaluateSkeletalControl_AnyThread(FComponentSpacePoseContext& Output, TArray<FBoneTransform>& OutBoneTransforms) override;
-	virtual bool IsValidToEvaluate(const USkeleton* Skeleton, const FBoneContainer& RequiredBones) override;
+	ANIMGRAPHRUNTIME_API virtual void EvaluateSkeletalControl_AnyThread(FComponentSpacePoseContext& Output, TArray<FBoneTransform>& OutBoneTransforms) override;
+	ANIMGRAPHRUNTIME_API virtual bool IsValidToEvaluate(const USkeleton* Skeleton, const FBoneContainer& RequiredBones) override;
 	// End of FAnimNode_SkeletalControlBase interface
 
 private:
 	// FAnimNode_SkeletalControlBase interface
-	virtual void InitializeBoneReferences(const FBoneContainer& RequiredBones) override;
+	ANIMGRAPHRUNTIME_API virtual void InitializeBoneReferences(const FBoneContainer& RequiredBones) override;
 	// End of FAnimNode_SkeletalControlBase interface
 };
diff --git a/Engine/Source/Runtime/AnimGraphRuntime/Public/BoneControllers/AnimNode_BoneDrivenController.h b/Engine/Source/Runtime/AnimGraphRuntime/Public/BoneControllers/AnimNode_BoneDrivenController.h
--- a/Engine/Source/Runtime/AnimGraphRuntime/Public/BoneControllers/AnimNode_BoneDrivenController.h
+++ b/Engine/Source/Runtime/AnimGraphRuntime/Public/BoneControllers/AnimNode_BoneDrivenController.h
@@ -43,7 +43,7 @@ enum class EDrivenDestinationMode : uint8
  * This is the runtime version of a bone driven controller, which maps part of the state from one bone to another (e.g., 2 * source.x -> target.z)
  */
 USTRUCT()
-struct ANIMGRAPHRUNTIME_API FAnimNode_BoneDrivenController : public FAnimNode_SkeletalControlBase
+struct FAnimNode_BoneDrivenController : public FAnimNode_SkeletalControlBase
 {
 	GENERATED_USTRUCT_BODY()
 
@@ -147,29 +147,29 @@ public:
 	uint8 bAffectTargetScaleZ : 1;
 
 public:
-	FAnimNode_BoneDrivenController();
+	ANIMGRAPHRUNTIME_API FAnimNode_BoneDrivenController();
 
 	// FAnimNode_Base interface
-	virtual void GatherDebugData(FNodeDebugData& DebugData) override;
+	ANIMGRAPHRUNTIME_API virtual void GatherDebugData(FNodeDebugData& DebugData) override;
 	// End of FAnimNode_Base interface
 
 	// FAnimNode_SkeletalControlBase interface
-	virtual void EvaluateSkeletalControl_AnyThread(FComponentSpacePoseContext& Output, TArray<FBoneTransform>& OutBoneTransforms) override;
-	virtual void EvaluateComponentSpaceInternal(FComponentSpacePoseContext& Context);
-	virtual bool IsValidToEvaluate(const USkeleton* Skeleton, const FBoneContainer& RequiredBones) override;
+	ANIMGRAPHRUNTIME_API virtual void EvaluateSkeletalControl_AnyThread(FComponentSpacePoseContext& Output, TArray<FBoneTransform>& OutBoneTransforms) override;
+	ANIMGRAPHRUNTIME_API virtual void EvaluateComponentSpaceInternal(FComponentSpacePoseContext& Context);
+	ANIMGRAPHRUNTIME_API virtual bool IsValidToEvaluate(const USkeleton* Skeleton, const FBoneContainer& RequiredBones) override;
 	// End of FAnimNode_SkeletalControlBase interface
 
 #if WITH_EDITORONLY_DATA
 	// Upgrade a node from the output enum to the output bits (change made in FAnimationCustomVersion::BoneDrivenControllerMatchingMaya)
-	void ConvertTargetComponentToBits();
+	ANIMGRAPHRUNTIME_API void ConvertTargetComponentToBits();
 #endif
 
 protected:
 	
 	// FAnimNode_SkeletalControlBase protected interface
-	virtual void InitializeBoneReferences(const FBoneContainer& RequiredBones) override;
+	ANIMGRAPHRUNTIME_API virtual void InitializeBoneReferences(const FBoneContainer& RequiredBones) override;
 
 	/** Extracts the value used to drive the target bone or parameter */
-	const double ExtractSourceValue(const FTransform& InCurrentBoneTransform, const FTransform& InRefPoseBoneTransform);
+	ANIMGRAPHRUNTIME_API const double ExtractSourceValue(const FTransform& InCurrentBoneTransform, const FTransform& InRefPoseBoneTransform);
 	// End of FAnimNode_SkeletalControlBase protected interface
 };
diff --git a/Engine/Source/Runtime/AnimGraphRuntime/Public/BoneControllers/AnimNode_CCDIK.h b/Engine/Source/Runtime/AnimGraphRuntime/Public/BoneControllers/AnimNode_CCDIK.h
--- a/Engine/Source/Runtime/AnimGraphRuntime/Public/BoneControllers/AnimNode_CCDIK.h
+++ b/Engine/Source/Runtime/AnimGraphRuntime/Public/BoneControllers/AnimNode_CCDIK.h
@@ -11,7 +11,7 @@
 *	Controller which implements the CCDIK IK approximation algorithm
 */
 USTRUCT()
-struct ANIMGRAPHRUNTIME_API FAnimNode_CCDIK : public FAnimNode_SkeletalControlBase
+struct FAnimNode_CCDIK : public FAnimNode_SkeletalControlBase
 {
 	GENERATED_USTRUCT_BODY()
 
@@ -57,20 +57,20 @@ private:
 	TArray<float> RotationLimitPerJoints;
 
 public:
-	FAnimNode_CCDIK();
+	ANIMGRAPHRUNTIME_API FAnimNode_CCDIK();
 
 	// FAnimNode_Base interface
-	virtual void GatherDebugData(FNodeDebugData& DebugData) override;
+	ANIMGRAPHRUNTIME_API virtual void GatherDebugData(FNodeDebugData& DebugData) override;
 	// End of FAnimNode_Base interface
 
 	// FAnimNode_SkeletalControlBase interface
-	virtual void EvaluateSkeletalControl_AnyThread(FComponentSpacePoseContext& Output, TArray<FBoneTransform>& OutBoneTransforms) override;
-	virtual bool IsValidToEvaluate(const USkeleton* Skeleton, const FBoneContainer& RequiredBones) override;
+	ANIMGRAPHRUNTIME_API virtual void EvaluateSkeletalControl_AnyThread(FComponentSpacePoseContext& Output, TArray<FBoneTransform>& OutBoneTransforms) override;
+	ANIMGRAPHRUNTIME_API virtual bool IsValidToEvaluate(const USkeleton* Skeleton, const FBoneContainer& RequiredBones) override;
 	// End of FAnimNode_SkeletalControlBase interface
 
 private:
 	// FAnimNode_SkeletalControlBase interface
-	virtual void InitializeBoneReferences(const FBoneContainer& RequiredBones) override;
+	ANIMGRAPHRUNTIME_API virtual void InitializeBoneReferences(const FBoneContainer& RequiredBones) override;
 	// End of FAnimNode_SkeletalControlBase interface
 
 	// Convenience function to get current (pre-translation iteration) component space location of bone by bone index
@@ -84,6 +84,6 @@ public:
 	TArray<FVector> DebugLines;
 #endif // #if UE_ENABLE_DEBUG_DRAWING
 	// resize rotation limit array based on set up
-	void ResizeRotationLimitPerJoints(int32 NewSize);
+	ANIMGRAPHRUNTIME_API void ResizeRotationLimitPerJoints(int32 NewSize);
 #endif
-};
\ No newline at end of file
+};
diff --git a/Engine/Source/Runtime/AnimGraphRuntime/Public/BoneControllers/AnimNode_Constraint.h b/Engine/Source/Runtime/AnimGraphRuntime/Public/BoneControllers/AnimNode_Constraint.h
--- a/Engine/Source/Runtime/AnimGraphRuntime/Public/BoneControllers/AnimNode_Constraint.h
+++ b/Engine/Source/Runtime/AnimGraphRuntime/Public/BoneControllers/AnimNode_Constraint.h
@@ -72,7 +72,7 @@ struct FConstraint
  *	Constraint node to parent or world transform for rotation/translation
  */
 USTRUCT()
-struct ANIMGRAPHRUNTIME_API FAnimNode_Constraint : public FAnimNode_SkeletalControlBase
+struct FAnimNode_Constraint : public FAnimNode_SkeletalControlBase
 {
 	GENERATED_USTRUCT_BODY()
 
@@ -88,24 +88,24 @@ struct ANIMGRAPHRUNTIME_API FAnimNode_Constraint : public FAnimNode_SkeletalCont
 	UPROPERTY(EditAnywhere, editfixedsize, Category = Runtime, meta = (PinShownByDefault))
 	TArray<float> ConstraintWeights;
 
-	FAnimNode_Constraint();
+	ANIMGRAPHRUNTIME_API FAnimNode_Constraint();
 
 	// FAnimNode_Base interface
-	virtual void GatherDebugData(FNodeDebugData& DebugData) override;
+	ANIMGRAPHRUNTIME_API virtual void GatherDebugData(FNodeDebugData& DebugData) override;
 	// End of FAnimNode_Base interface
 
 	// FAnimNode_SkeletalControlBase interface
-	virtual void EvaluateSkeletalControl_AnyThread(FComponentSpacePoseContext& Output, TArray<FBoneTransform>& OutBoneTransforms) override;
-	virtual bool IsValidToEvaluate(const USkeleton* Skeleton, const FBoneContainer& RequiredBones) override;
+	ANIMGRAPHRUNTIME_API virtual void EvaluateSkeletalControl_AnyThread(FComponentSpacePoseContext& Output, TArray<FBoneTransform>& OutBoneTransforms) override;
+	ANIMGRAPHRUNTIME_API virtual bool IsValidToEvaluate(const USkeleton* Skeleton, const FBoneContainer& RequiredBones) override;
 	// End of FAnimNode_SkeletalControlBase interface
 
 #if WITH_EDITOR
-	void ConditionalDebugDraw(FPrimitiveDrawInterface* PDI, USkeletalMeshComponent* MeshComp) const;
+	ANIMGRAPHRUNTIME_API void ConditionalDebugDraw(FPrimitiveDrawInterface* PDI, USkeletalMeshComponent* MeshComp) const;
 #endif // WITH_EDITOR
 
 private:
 	// FAnimNode_SkeletalControlBase interface
-	virtual void InitializeBoneReferences(const FBoneContainer& RequiredBones) override;
+	ANIMGRAPHRUNTIME_API virtual void InitializeBoneReferences(const FBoneContainer& RequiredBones) override;
 	// End of FAnimNode_SkeletalControlBase interface
 
 	TArray<FConstraintData>	ConstraintData;
diff --git a/Engine/Source/Runtime/AnimGraphRuntime/Public/BoneControllers/AnimNode_CopyBone.h b/Engine/Source/Runtime/AnimGraphRuntime/Public/BoneControllers/AnimNode_CopyBone.h
--- a/Engine/Source/Runtime/AnimGraphRuntime/Public/BoneControllers/AnimNode_CopyBone.h
+++ b/Engine/Source/Runtime/AnimGraphRuntime/Public/BoneControllers/AnimNode_CopyBone.h
@@ -16,7 +16,7 @@ class USkeletalMeshComponent;
  */
 
 USTRUCT(BlueprintInternalUseOnly)
-struct ANIMGRAPHRUNTIME_API FAnimNode_CopyBone : public FAnimNode_SkeletalControlBase
+struct FAnimNode_CopyBone : public FAnimNode_SkeletalControlBase
 {
 	GENERATED_USTRUCT_BODY()
 
@@ -44,19 +44,19 @@ struct ANIMGRAPHRUNTIME_API FAnimNode_CopyBone : public FAnimNode_SkeletalContro
 	UPROPERTY(EditAnywhere, Category = Copy)
 	TEnumAsByte<EBoneControlSpace> ControlSpace;
 
-	FAnimNode_CopyBone();
+	ANIMGRAPHRUNTIME_API FAnimNode_CopyBone();
 
 	// FAnimNode_Base interface
-	virtual void GatherDebugData(FNodeDebugData& DebugData) override;
+	ANIMGRAPHRUNTIME_API virtual void GatherDebugData(FNodeDebugData& DebugData) override;
 	// End of FAnimNode_Base interface
 
 	// FAnimNode_SkeletalControlBase interface
-	virtual void EvaluateSkeletalControl_AnyThread(FComponentSpacePoseContext& Output, TArray<FBoneTransform>& OutBoneTransforms) override;
-	virtual bool IsValidToEvaluate(const USkeleton* Skeleton, const FBoneContainer& RequiredBones) override;
+	ANIMGRAPHRUNTIME_API virtual void EvaluateSkeletalControl_AnyThread(FComponentSpacePoseContext& Output, TArray<FBoneTransform>& OutBoneTransforms) override;
+	ANIMGRAPHRUNTIME_API virtual bool IsValidToEvaluate(const USkeleton* Skeleton, const FBoneContainer& RequiredBones) override;
 	// End of FAnimNode_SkeletalControlBase interface
 
 private:
 	// FAnimNode_SkeletalControlBase interface
-	virtual void InitializeBoneReferences(const FBoneContainer& RequiredBones) override;
+	ANIMGRAPHRUNTIME_API virtual void InitializeBoneReferences(const FBoneContainer& RequiredBones) override;
 	// End of FAnimNode_SkeletalControlBase interface
 };
diff --git a/Engine/Source/Runtime/AnimGraphRuntime/Public/BoneControllers/AnimNode_CopyBoneDelta.h b/Engine/Source/Runtime/AnimGraphRuntime/Public/BoneControllers/AnimNode_CopyBoneDelta.h
--- a/Engine/Source/Runtime/AnimGraphRuntime/Public/BoneControllers/AnimNode_CopyBoneDelta.h
+++ b/Engine/Source/Runtime/AnimGraphRuntime/Public/BoneControllers/AnimNode_CopyBoneDelta.h
@@ -24,7 +24,7 @@ enum class CopyBoneDeltaMode : uint8
  */
 
 USTRUCT(BlueprintInternalUseOnly)
-struct ANIMGRAPHRUNTIME_API FAnimNode_CopyBoneDelta : public FAnimNode_SkeletalControlBase
+struct FAnimNode_CopyBoneDelta : public FAnimNode_SkeletalControlBase
 {
 	GENERATED_USTRUCT_BODY()
 
@@ -55,19 +55,19 @@ struct ANIMGRAPHRUNTIME_API FAnimNode_CopyBoneDelta : public FAnimNode_SkeletalC
 	UPROPERTY(EditAnywhere, BlueprintReadWrite, Category = Settings, meta = (PinHiddenByDefault))
 	float ScaleMultiplier;
 
-	FAnimNode_CopyBoneDelta();
+	ANIMGRAPHRUNTIME_API FAnimNode_CopyBoneDelta();
 
 	// FAnimNode_Base interface
-	virtual void GatherDebugData(FNodeDebugData& DebugData) override;
+	ANIMGRAPHRUNTIME_API virtual void GatherDebugData(FNodeDebugData& DebugData) override;
 	// End of FAnimNode_Base interface
 
 	// FAnimNode_SkeletalControlBase interface
-	virtual void EvaluateSkeletalControl_AnyThread(FComponentSpacePoseContext& Output, TArray<FBoneTransform>& OutBoneTransforms) override;
-	virtual bool IsValidToEvaluate(const USkeleton* Skeleton, const FBoneContainer& RequiredBones) override;
+	ANIMGRAPHRUNTIME_API virtual void EvaluateSkeletalControl_AnyThread(FComponentSpacePoseContext& Output, TArray<FBoneTransform>& OutBoneTransforms) override;
+	ANIMGRAPHRUNTIME_API virtual bool IsValidToEvaluate(const USkeleton* Skeleton, const FBoneContainer& RequiredBones) override;
 	// End of FAnimNode_SkeletalControlBase interface
 
 private:
 	// FAnimNode_SkeletalControlBase interface
-	virtual void InitializeBoneReferences(const FBoneContainer& RequiredBones) override;
+	ANIMGRAPHRUNTIME_API virtual void InitializeBoneReferences(const FBoneContainer& RequiredBones) override;
 	// End of FAnimNode_SkeletalControlBase interface
 };
diff --git a/Engine/Source/Runtime/AnimGraphRuntime/Public/BoneControllers/AnimNode_Fabrik.h b/Engine/Source/Runtime/AnimGraphRuntime/Public/BoneControllers/AnimNode_Fabrik.h
--- a/Engine/Source/Runtime/AnimGraphRuntime/Public/BoneControllers/AnimNode_Fabrik.h
+++ b/Engine/Source/Runtime/AnimGraphRuntime/Public/BoneControllers/AnimNode_Fabrik.h
@@ -18,7 +18,7 @@ class USkeletalMeshComponent;
 */
 
 USTRUCT(BlueprintInternalUseOnly)
-struct ANIMGRAPHRUNTIME_API FAnimNode_Fabrik : public FAnimNode_SkeletalControlBase
+struct FAnimNode_Fabrik : public FAnimNode_SkeletalControlBase
 {
 	GENERATED_USTRUCT_BODY()
 
@@ -64,23 +64,23 @@ struct ANIMGRAPHRUNTIME_API FAnimNode_Fabrik : public FAnimNode_SkeletalControlB
 #endif
 
 public:
-	FAnimNode_Fabrik();
+	ANIMGRAPHRUNTIME_API FAnimNode_Fabrik();
 
 	// FAnimNode_Base interface
-	virtual void GatherDebugData(FNodeDebugData& DebugData) override;
-	virtual void Initialize_AnyThread(const FAnimationInitializeContext& Context) override;
+	ANIMGRAPHRUNTIME_API virtual void GatherDebugData(FNodeDebugData& DebugData) override;
+	ANIMGRAPHRUNTIME_API virtual void Initialize_AnyThread(const FAnimationInitializeContext& Context) override;
 	// End of FAnimNode_Base interface
 
 	// FAnimNode_SkeletalControlBase interface
-	virtual void EvaluateSkeletalControl_AnyThread(FComponentSpacePoseContext& Output, TArray<FBoneTransform>& OutBoneTransforms) override;
-	virtual bool IsValidToEvaluate(const USkeleton* Skeleton, const FBoneContainer& RequiredBones) override;
+	ANIMGRAPHRUNTIME_API virtual void EvaluateSkeletalControl_AnyThread(FComponentSpacePoseContext& Output, TArray<FBoneTransform>& OutBoneTransforms) override;
+	ANIMGRAPHRUNTIME_API virtual bool IsValidToEvaluate(const USkeleton* Skeleton, const FBoneContainer& RequiredBones) override;
 	// End of FAnimNode_SkeletalControlBase interface
 
-	virtual void ConditionalDebugDraw(FPrimitiveDrawInterface* PDI, USkeletalMeshComponent* PreviewSkelMeshComp) const;
+	ANIMGRAPHRUNTIME_API virtual void ConditionalDebugDraw(FPrimitiveDrawInterface* PDI, USkeletalMeshComponent* PreviewSkelMeshComp) const;
 
 private:
 	// FAnimNode_SkeletalControlBase interface
-	virtual void InitializeBoneReferences(const FBoneContainer& RequiredBones) override;
+	ANIMGRAPHRUNTIME_API virtual void InitializeBoneReferences(const FBoneContainer& RequiredBones) override;
 	// End of FAnimNode_SkeletalControlBase interface
 
 	// Convenience function to get current (pre-translation iteration) component space location of bone by bone index
diff --git a/Engine/Source/Runtime/AnimGraphRuntime/Public/BoneControllers/AnimNode_HandIKRetargeting.h b/Engine/Source/Runtime/AnimGraphRuntime/Public/BoneControllers/AnimNode_HandIKRetargeting.h
--- a/Engine/Source/Runtime/AnimGraphRuntime/Public/BoneControllers/AnimNode_HandIKRetargeting.h
+++ b/Engine/Source/Runtime/AnimGraphRuntime/Public/BoneControllers/AnimNode_HandIKRetargeting.h
@@ -19,7 +19,7 @@ class USkeletalMeshComponent;
  */
 
 USTRUCT(BlueprintInternalUseOnly)
-struct ANIMGRAPHRUNTIME_API FAnimNode_HandIKRetargeting : public FAnimNode_SkeletalControlBase
+struct FAnimNode_HandIKRetargeting : public FAnimNode_SkeletalControlBase
 {
 	GENERATED_USTRUCT_BODY()
 
@@ -47,19 +47,19 @@ struct ANIMGRAPHRUNTIME_API FAnimNode_HandIKRetargeting : public FAnimNode_Skele
 	UPROPERTY(EditAnywhere, BlueprintReadWrite, Category = "HandIKRetargeting", meta = (PinShownByDefault))
 	float HandFKWeight;
 
-	FAnimNode_HandIKRetargeting();
+	ANIMGRAPHRUNTIME_API FAnimNode_HandIKRetargeting();
 
 	// FAnimNode_Base interface
-	virtual void GatherDebugData(FNodeDebugData& DebugData) override;
+	ANIMGRAPHRUNTIME_API virtual void GatherDebugData(FNodeDebugData& DebugData) override;
 	// End of FAnimNode_Base interface
 
 	// FAnimNode_SkeletalControlBase interface
-	virtual void EvaluateSkeletalControl_AnyThread(FComponentSpacePoseContext& Output, TArray<FBoneTransform>& OutBoneTransforms) override;
-	virtual bool IsValidToEvaluate(const USkeleton* Skeleton, const FBoneContainer& RequiredBones) override;
+	ANIMGRAPHRUNTIME_API virtual void EvaluateSkeletalControl_AnyThread(FComponentSpacePoseContext& Output, TArray<FBoneTransform>& OutBoneTransforms) override;
+	ANIMGRAPHRUNTIME_API virtual bool IsValidToEvaluate(const USkeleton* Skeleton, const FBoneContainer& RequiredBones) override;
 	// End of FAnimNode_SkeletalControlBase interface
 
 private:
 	// FAnimNode_SkeletalControlBase interface
-	virtual void InitializeBoneReferences(const FBoneContainer& RequiredBones) override;
+	ANIMGRAPHRUNTIME_API virtual void InitializeBoneReferences(const FBoneContainer& RequiredBones) override;
 	// End of FAnimNode_SkeletalControlBase interface
 };
diff --git a/Engine/Source/Runtime/AnimGraphRuntime/Public/BoneControllers/AnimNode_LegIK.h b/Engine/Source/Runtime/AnimGraphRuntime/Public/BoneControllers/AnimNode_LegIK.h
--- a/Engine/Source/Runtime/AnimGraphRuntime/Public/BoneControllers/AnimNode_LegIK.h
+++ b/Engine/Source/Runtime/AnimGraphRuntime/Public/BoneControllers/AnimNode_LegIK.h
@@ -177,11 +177,11 @@ public:
 };
 
 USTRUCT()
-struct ANIMGRAPHRUNTIME_API FAnimNode_LegIK : public FAnimNode_SkeletalControlBase
+struct FAnimNode_LegIK : public FAnimNode_SkeletalControlBase
 {
 	GENERATED_USTRUCT_BODY()
 
-	FAnimNode_LegIK();
+	ANIMGRAPHRUNTIME_API FAnimNode_LegIK();
 
 	/** Tolerance for reaching IK Target, in unreal units. */
 	UPROPERTY(EditAnywhere, Category = "Settings")
@@ -200,21 +200,21 @@ struct ANIMGRAPHRUNTIME_API FAnimNode_LegIK : public FAnimNode_SkeletalControlBa
 
 public:
 	// FAnimNode_Base interface
-	virtual void GatherDebugData(FNodeDebugData& DebugData) override;
+	ANIMGRAPHRUNTIME_API virtual void GatherDebugData(FNodeDebugData& DebugData) override;
 	// End of FAnimNode_Base interface
 
 	// FAnimNode_SkeletalControlBase interface
-	virtual void Initialize_AnyThread(const FAnimationInitializeContext& Context) override;
-	virtual void EvaluateSkeletalControl_AnyThread(FComponentSpacePoseContext& Output, TArray<FBoneTransform>& OutBoneTransforms) override;
-	virtual bool IsValidToEvaluate(const USkeleton* Skeleton, const FBoneContainer& RequiredBones) override;
+	ANIMGRAPHRUNTIME_API virtual void Initialize_AnyThread(const FAnimationInitializeContext& Context) override;
+	ANIMGRAPHRUNTIME_API virtual void EvaluateSkeletalControl_AnyThread(FComponentSpacePoseContext& Output, TArray<FBoneTransform>& OutBoneTransforms) override;
+	ANIMGRAPHRUNTIME_API virtual bool IsValidToEvaluate(const USkeleton* Skeleton, const FBoneContainer& RequiredBones) override;
 	// End of FAnimNode_SkeletalControlBase interface
 
-	bool OrientLegTowardsIK(FAnimLegIKData& InLegData);
-	bool DoLegReachIK(FAnimLegIKData& InLegData);
-	bool AdjustKneeTwist(FAnimLegIKData& InLegData);
+	ANIMGRAPHRUNTIME_API bool OrientLegTowardsIK(FAnimLegIKData& InLegData);
+	ANIMGRAPHRUNTIME_API bool DoLegReachIK(FAnimLegIKData& InLegData);
+	ANIMGRAPHRUNTIME_API bool AdjustKneeTwist(FAnimLegIKData& InLegData);
 
 private:
 	// FAnimNode_SkeletalControlBase interface
-	virtual void InitializeBoneReferences(const FBoneContainer& RequiredBones) override;
+	ANIMGRAPHRUNTIME_API virtual void InitializeBoneReferences(const FBoneContainer& RequiredBones) override;
 	// End of FAnimNode_SkeletalControlBase interface
 };
diff --git a/Engine/Source/Runtime/AnimGraphRuntime/Public/BoneControllers/AnimNode_LookAt.h b/Engine/Source/Runtime/AnimGraphRuntime/Public/BoneControllers/AnimNode_LookAt.h
--- a/Engine/Source/Runtime/AnimGraphRuntime/Public/BoneControllers/AnimNode_LookAt.h
+++ b/Engine/Source/Runtime/AnimGraphRuntime/Public/BoneControllers/AnimNode_LookAt.h
@@ -36,7 +36,7 @@ namespace EInterpolationBlend
  *	Simple controller that make a bone to look at the point or another bone
  */
 USTRUCT(BlueprintInternalUseOnly)
-struct ANIMGRAPHRUNTIME_API FAnimNode_LookAt : public FAnimNode_SkeletalControlBase
+struct FAnimNode_LookAt : public FAnimNode_SkeletalControlBase
 {
 	GENERATED_USTRUCT_BODY()
 
@@ -101,33 +101,33 @@ struct ANIMGRAPHRUNTIME_API FAnimNode_LookAt : public FAnimNode_SkeletalControlB
 #endif
 
 	// in the future, it would be nice to have more options, -i.e. lag, interpolation speed
-	FAnimNode_LookAt();
+	ANIMGRAPHRUNTIME_API FAnimNode_LookAt();
 
 	// FAnimNode_Base interface
-	virtual void GatherDebugData(FNodeDebugData& DebugData) override;
-	virtual void UpdateInternal(const FAnimationUpdateContext& Context) override;
-	virtual void Initialize_AnyThread(const FAnimationInitializeContext& Context) override;
+	ANIMGRAPHRUNTIME_API virtual void GatherDebugData(FNodeDebugData& DebugData) override;
+	ANIMGRAPHRUNTIME_API virtual void UpdateInternal(const FAnimationUpdateContext& Context) override;
+	ANIMGRAPHRUNTIME_API virtual void Initialize_AnyThread(const FAnimationInitializeContext& Context) override;
 	// End of FAnimNode_Base interface
 
 	// FAnimNode_SkeletalControlBase interface
-	virtual void EvaluateComponentSpaceInternal(FComponentSpacePoseContext& Context) override;
-	virtual void EvaluateSkeletalControl_AnyThread(FComponentSpacePoseContext& Output, TArray<FBoneTransform>& OutBoneTransforms) override;
-	virtual bool IsValidToEvaluate(const USkeleton* Skeleton, const FBoneContainer& RequiredBones) override;
+	ANIMGRAPHRUNTIME_API virtual void EvaluateComponentSpaceInternal(FComponentSpacePoseContext& Context) override;
+	ANIMGRAPHRUNTIME_API virtual void EvaluateSkeletalControl_AnyThread(FComponentSpacePoseContext& Output, TArray<FBoneTransform>& OutBoneTransforms) override;
+	ANIMGRAPHRUNTIME_API virtual bool IsValidToEvaluate(const USkeleton* Skeleton, const FBoneContainer& RequiredBones) override;
 	// End of FAnimNode_SkeletalControlBase interface
 
 	FVector GetCachedTargetLocation() const {	return 	CachedCurrentTargetLocation;	}
 
 #if WITH_EDITOR
-	void ConditionalDebugDraw(FPrimitiveDrawInterface* PDI, USkeletalMeshComponent* MeshComp) const;
+	ANIMGRAPHRUNTIME_API void ConditionalDebugDraw(FPrimitiveDrawInterface* PDI, USkeletalMeshComponent* MeshComp) const;
 #endif // WITH_EDITOR
 
 private:
 	// FAnimNode_SkeletalControlBase interface
-	virtual void InitializeBoneReferences(const FBoneContainer& RequiredBones) override;
+	ANIMGRAPHRUNTIME_API virtual void InitializeBoneReferences(const FBoneContainer& RequiredBones) override;
 	// End of FAnimNode_SkeletalControlBase interface
 
 	/** Turn a linear interpolated alpha into the corresponding AlphaBlendType */
-	static float AlphaToBlendType(float InAlpha, EInterpolationBlend::Type BlendType);
+	static ANIMGRAPHRUNTIME_API float AlphaToBlendType(float InAlpha, EInterpolationBlend::Type BlendType);
 
 	/** Debug transient data */
 	FVector CurrentLookAtLocation;
@@ -151,5 +151,5 @@ private:
 	FVector CachedCurrentTargetLocation;
 
 protected:
-	virtual void ModifyPoseFromDeltaRotation(FComponentSpacePoseContext& Output, TArray<FBoneTransform>& OutBoneTransforms,	FTransform& InOutBoneToModifyTransform,	const FQuat& DeltaRotation);
+	ANIMGRAPHRUNTIME_API virtual void ModifyPoseFromDeltaRotation(FComponentSpacePoseContext& Output, TArray<FBoneTransform>& OutBoneTransforms,	FTransform& InOutBoneToModifyTransform,	const FQuat& DeltaRotation);
 };
diff --git a/Engine/Source/Runtime/AnimGraphRuntime/Public/BoneControllers/AnimNode_ModifyBone.h b/Engine/Source/Runtime/AnimGraphRuntime/Public/BoneControllers/AnimNode_ModifyBone.h
--- a/Engine/Source/Runtime/AnimGraphRuntime/Public/BoneControllers/AnimNode_ModifyBone.h
+++ b/Engine/Source/Runtime/AnimGraphRuntime/Public/BoneControllers/AnimNode_ModifyBone.h
@@ -28,7 +28,7 @@ enum EBoneModificationMode : int
  *	Simple controller that replaces or adds to the translation/rotation of a single bone.
  */
 USTRUCT(BlueprintInternalUseOnly)
-struct ANIMGRAPHRUNTIME_API FAnimNode_ModifyBone : public FAnimNode_SkeletalControlBase
+struct FAnimNode_ModifyBone : public FAnimNode_SkeletalControlBase
 {
 	GENERATED_USTRUCT_BODY()
 
@@ -72,19 +72,19 @@ struct ANIMGRAPHRUNTIME_API FAnimNode_ModifyBone : public FAnimNode_SkeletalCont
 	UPROPERTY(EditAnywhere, Category=Scale)
 	TEnumAsByte<enum EBoneControlSpace> ScaleSpace;
 
-	FAnimNode_ModifyBone();
+	ANIMGRAPHRUNTIME_API FAnimNode_ModifyBone();
 
 	// FAnimNode_Base interface
-	virtual void GatherDebugData(FNodeDebugData& DebugData) override;
+	ANIMGRAPHRUNTIME_API virtual void GatherDebugData(FNodeDebugData& DebugData) override;
 	// End of FAnimNode_Base interface
 
 	// FAnimNode_SkeletalControlBase interface
-	virtual void EvaluateSkeletalControl_AnyThread(FComponentSpacePoseContext& Output, TArray<FBoneTransform>& OutBoneTransforms) override;
-	virtual bool IsValidToEvaluate(const USkeleton* Skeleton, const FBoneContainer& RequiredBones) override;
+	ANIMGRAPHRUNTIME_API virtual void EvaluateSkeletalControl_AnyThread(FComponentSpacePoseContext& Output, TArray<FBoneTransform>& OutBoneTransforms) override;
+	ANIMGRAPHRUNTIME_API virtual bool IsValidToEvaluate(const USkeleton* Skeleton, const FBoneContainer& RequiredBones) override;
 	// End of FAnimNode_SkeletalControlBase interface
 
 private:
 	// FAnimNode_SkeletalControlBase interface
-	virtual void InitializeBoneReferences(const FBoneContainer& RequiredBones) override;
+	ANIMGRAPHRUNTIME_API virtual void InitializeBoneReferences(const FBoneContainer& RequiredBones) override;
 	// End of FAnimNode_SkeletalControlBase interface
 };
diff --git a/Engine/Source/Runtime/AnimGraphRuntime/Public/BoneControllers/AnimNode_ObserveBone.h b/Engine/Source/Runtime/AnimGraphRuntime/Public/BoneControllers/AnimNode_ObserveBone.h
--- a/Engine/Source/Runtime/AnimGraphRuntime/Public/BoneControllers/AnimNode_ObserveBone.h
+++ b/Engine/Source/Runtime/AnimGraphRuntime/Public/BoneControllers/AnimNode_ObserveBone.h
@@ -15,7 +15,7 @@ class USkeletalMeshComponent;
  *	Debugging node that displays the current value of a bone in a specific space.
  */
 USTRUCT()
-struct ANIMGRAPHRUNTIME_API FAnimNode_ObserveBone : public FAnimNode_SkeletalControlBase
+struct FAnimNode_ObserveBone : public FAnimNode_SkeletalControlBase
 {
 	GENERATED_USTRUCT_BODY()
 
@@ -44,19 +44,19 @@ struct ANIMGRAPHRUNTIME_API FAnimNode_ObserveBone : public FAnimNode_SkeletalCon
 	FVector Scale;
 
 public:
-	FAnimNode_ObserveBone();
+	ANIMGRAPHRUNTIME_API FAnimNode_ObserveBone();
 
 	// FAnimNode_Base interface
-	virtual void GatherDebugData(FNodeDebugData& DebugData) override;
+	ANIMGRAPHRUNTIME_API virtual void GatherDebugData(FNodeDebugData& DebugData) override;
 	// End of FAnimNode_Base interface
 
 	// FAnimNode_SkeletalControlBase interface
-	virtual void EvaluateSkeletalControl_AnyThread(FComponentSpacePoseContext& Output, TArray<FBoneTransform>& OutBoneTransforms) override;
-	virtual bool IsValidToEvaluate(const USkeleton* Skeleton, const FBoneContainer& RequiredBones) override;
+	ANIMGRAPHRUNTIME_API virtual void EvaluateSkeletalControl_AnyThread(FComponentSpacePoseContext& Output, TArray<FBoneTransform>& OutBoneTransforms) override;
+	ANIMGRAPHRUNTIME_API virtual bool IsValidToEvaluate(const USkeleton* Skeleton, const FBoneContainer& RequiredBones) override;
 	// End of FAnimNode_SkeletalControlBase interface
 
 private:
 	// FAnimNode_SkeletalControlBase interface
-	virtual void InitializeBoneReferences(const FBoneContainer& RequiredBones) override;
+	ANIMGRAPHRUNTIME_API virtual void InitializeBoneReferences(const FBoneContainer& RequiredBones) override;
 	// End of FAnimNode_SkeletalControlBase interface
 };
diff --git a/Engine/Source/Runtime/AnimGraphRuntime/Public/BoneControllers/AnimNode_ResetRoot.h b/Engine/Source/Runtime/AnimGraphRuntime/Public/BoneControllers/AnimNode_ResetRoot.h
--- a/Engine/Source/Runtime/AnimGraphRuntime/Public/BoneControllers/AnimNode_ResetRoot.h
+++ b/Engine/Source/Runtime/AnimGraphRuntime/Public/BoneControllers/AnimNode_ResetRoot.h
@@ -10,27 +10,27 @@
 #include "AnimNode_ResetRoot.generated.h"
 
 USTRUCT(BlueprintInternalUseOnly)
-struct ANIMGRAPHRUNTIME_API FAnimNode_ResetRoot : public FAnimNode_SkeletalControlBase
+struct FAnimNode_ResetRoot : public FAnimNode_SkeletalControlBase
 {
 	GENERATED_USTRUCT_BODY();
 
-	FAnimNode_ResetRoot();
+	ANIMGRAPHRUNTIME_API FAnimNode_ResetRoot();
 
 	TArray<FCompactPoseBoneIndex> RootChildren;
 
 	// FAnimNode_Base interface
-	virtual void GatherDebugData(FNodeDebugData& DebugData) override;
-	virtual void OnInitializeAnimInstance(const FAnimInstanceProxy* InProxy, const UAnimInstance* InAnimInstance) override;
+	ANIMGRAPHRUNTIME_API virtual void GatherDebugData(FNodeDebugData& DebugData) override;
+	ANIMGRAPHRUNTIME_API virtual void OnInitializeAnimInstance(const FAnimInstanceProxy* InProxy, const UAnimInstance* InAnimInstance) override;
 	virtual bool NeedsOnInitializeAnimInstance() const override { return true; }
 	// End of FAnimNode_Base interface
 
 	// FAnimNode_SkeletalControlBase interface
-	virtual void EvaluateSkeletalControl_AnyThread(FComponentSpacePoseContext& Output, TArray<FBoneTransform>& OutBoneTransforms) override;
-	virtual bool IsValidToEvaluate(const USkeleton* Skeleton, const FBoneContainer& RequiredBones) override;
+	ANIMGRAPHRUNTIME_API virtual void EvaluateSkeletalControl_AnyThread(FComponentSpacePoseContext& Output, TArray<FBoneTransform>& OutBoneTransforms) override;
+	ANIMGRAPHRUNTIME_API virtual bool IsValidToEvaluate(const USkeleton* Skeleton, const FBoneContainer& RequiredBones) override;
 	// End of FAnimNode_SkeletalControlBase interface
 
 private:
 	// FAnimNode_SkeletalControlBase interface
-	virtual void InitializeBoneReferences(const FBoneContainer& RequiredBones) override;
+	ANIMGRAPHRUNTIME_API virtual void InitializeBoneReferences(const FBoneContainer& RequiredBones) override;
 	// End of FAnimNode_SkeletalControlBase interface
 };
diff --git a/Engine/Source/Runtime/AnimGraphRuntime/Public/BoneControllers/AnimNode_RigidBody.h b/Engine/Source/Runtime/AnimGraphRuntime/Public/BoneControllers/AnimNode_RigidBody.h
--- a/Engine/Source/Runtime/AnimGraphRuntime/Public/BoneControllers/AnimNode_RigidBody.h
+++ b/Engine/Source/Runtime/AnimGraphRuntime/Public/BoneControllers/AnimNode_RigidBody.h
@@ -50,11 +50,11 @@ enum class ESimulationTiming : uint8
 template <> struct TIsPODType<FSimSpaceSettings> { enum { Value = true }; };
 
 USTRUCT(BlueprintType)
-struct ANIMGRAPHRUNTIME_API FSimSpaceSettings
+struct FSimSpaceSettings
 {
 	GENERATED_USTRUCT_BODY()
 
-	FSimSpaceSettings();
+	ANIMGRAPHRUNTIME_API FSimSpaceSettings();
 
 	// Disable deprecation errors by providing defaults wrapped with pragma disable
 	PRAGMA_DISABLE_DEPRECATION_WARNINGS
@@ -127,7 +127,7 @@ struct ANIMGRAPHRUNTIME_API FSimSpaceSettings
 	UPROPERTY(EditAnywhere, BlueprintReadWrite, Category = Settings)
 	FVector ExternalAngularVelocity;
 
-	void PostSerialize(const FArchive& Ar);
+	ANIMGRAPHRUNTIME_API void PostSerialize(const FArchive& Ar);
 };
 
 #if WITH_EDITORONLY_DATA
@@ -146,34 +146,34 @@ struct TStructOpsTypeTraits<FSimSpaceSettings> : public TStructOpsTypeTraitsBase
  *	Controller that simulates physics based on the physics asset of the skeletal mesh component
  */
 USTRUCT()
-struct ANIMGRAPHRUNTIME_API FAnimNode_RigidBody : public FAnimNode_SkeletalControlBase
+struct FAnimNode_RigidBody : public FAnimNode_SkeletalControlBase
 {
 	GENERATED_USTRUCT_BODY()
 
-	FAnimNode_RigidBody();
-	~FAnimNode_RigidBody();
+	ANIMGRAPHRUNTIME_API FAnimNode_RigidBody();
+	ANIMGRAPHRUNTIME_API ~FAnimNode_RigidBody();
 
 	// FAnimNode_Base interface
-	virtual void GatherDebugData(FNodeDebugData& DebugData) override;
-	virtual void Initialize_AnyThread(const FAnimationInitializeContext& Context) override;
+	ANIMGRAPHRUNTIME_API virtual void GatherDebugData(FNodeDebugData& DebugData) override;
+	ANIMGRAPHRUNTIME_API virtual void Initialize_AnyThread(const FAnimationInitializeContext& Context) override;
 	// End of FAnimNode_Base interface
 
 	// FAnimNode_SkeletalControlBase interface
-	virtual void UpdateComponentPose_AnyThread(const FAnimationUpdateContext& Context) override;
-	virtual void EvaluateComponentPose_AnyThread(FComponentSpacePoseContext& Output) override;
-	virtual void EvaluateSkeletalControl_AnyThread(FComponentSpacePoseContext& Output, TArray<FBoneTransform>& OutBoneTransforms) override;
-	virtual void OnInitializeAnimInstance(const FAnimInstanceProxy* InProxy, const UAnimInstance* InAnimInstance) override;
+	ANIMGRAPHRUNTIME_API virtual void UpdateComponentPose_AnyThread(const FAnimationUpdateContext& Context) override;
+	ANIMGRAPHRUNTIME_API virtual void EvaluateComponentPose_AnyThread(FComponentSpacePoseContext& Output) override;
+	ANIMGRAPHRUNTIME_API virtual void EvaluateSkeletalControl_AnyThread(FComponentSpacePoseContext& Output, TArray<FBoneTransform>& OutBoneTransforms) override;
+	ANIMGRAPHRUNTIME_API virtual void OnInitializeAnimInstance(const FAnimInstanceProxy* InProxy, const UAnimInstance* InAnimInstance) override;
 	virtual bool NeedsOnInitializeAnimInstance() const override { return true; }
-	virtual void PreUpdate(const UAnimInstance* InAnimInstance) override;
-	virtual void UpdateInternal(const FAnimationUpdateContext& Context) override;
+	ANIMGRAPHRUNTIME_API virtual void PreUpdate(const UAnimInstance* InAnimInstance) override;
+	ANIMGRAPHRUNTIME_API virtual void UpdateInternal(const FAnimationUpdateContext& Context) override;
 	virtual bool HasPreUpdate() const override { return true; }
-	virtual bool IsValidToEvaluate(const USkeleton* Skeleton, const FBoneContainer& RequiredBones) override;
-	virtual bool NeedsDynamicReset() const override;
-	virtual void ResetDynamics(ETeleportType InTeleportType) override;
-	virtual int32 GetLODThreshold() const override;
+	ANIMGRAPHRUNTIME_API virtual bool IsValidToEvaluate(const USkeleton* Skeleton, const FBoneContainer& RequiredBones) override;
+	ANIMGRAPHRUNTIME_API virtual bool NeedsDynamicReset() const override;
+	ANIMGRAPHRUNTIME_API virtual void ResetDynamics(ETeleportType InTeleportType) override;
+	ANIMGRAPHRUNTIME_API virtual int32 GetLODThreshold() const override;
 	// End of FAnimNode_SkeletalControlBase interface
 
-	virtual void AddImpulseAtLocation(FVector Impulse, FVector Location, FName BoneName = NAME_None);
+	ANIMGRAPHRUNTIME_API virtual void AddImpulseAtLocation(FVector Impulse, FVector Location, FName BoneName = NAME_None);
 
 	// TEMP: Exposed for use in PhAt as a quick way to get drag handles working with Chaos
 	virtual ImmediatePhysics::FSimulation* GetSimulation() { return PhysicsSimulation; }
@@ -305,7 +305,7 @@ private:
 	uint8 bCheckForBodyTransformInit : 1;
 
 public:
-	void PostSerialize(const FArchive& Ar);
+	ANIMGRAPHRUNTIME_API void PostSerialize(const FArchive& Ar);
 
 private:
 
@@ -315,21 +315,21 @@ private:
 #endif
 
 	// FAnimNode_SkeletalControlBase interface
-	virtual void InitializeBoneReferences(const FBoneContainer& RequiredBones) override;
+	ANIMGRAPHRUNTIME_API virtual void InitializeBoneReferences(const FBoneContainer& RequiredBones) override;
 	// End of FAnimNode_SkeletalControlBase interface
 
-	void InitPhysics(const UAnimInstance* InAnimInstance);
-	void UpdateWorldGeometry(const UWorld& World, const USkeletalMeshComponent& SKC);
-	void UpdateWorldForces(const FTransform& ComponentToWorld, const FTransform& RootBoneTM, const float DeltaSeconds);
+	ANIMGRAPHRUNTIME_API void InitPhysics(const UAnimInstance* InAnimInstance);
+	ANIMGRAPHRUNTIME_API void UpdateWorldGeometry(const UWorld& World, const USkeletalMeshComponent& SKC);
+	ANIMGRAPHRUNTIME_API void UpdateWorldForces(const FTransform& ComponentToWorld, const FTransform& RootBoneTM, const float DeltaSeconds);
 
-	void InitializeNewBodyTransformsDuringSimulation(FComponentSpacePoseContext& Output, const FTransform& ComponentTransform, const FTransform& BaseBoneTM);
+	ANIMGRAPHRUNTIME_API void InitializeNewBodyTransformsDuringSimulation(FComponentSpacePoseContext& Output, const FTransform& ComponentTransform, const FTransform& BaseBoneTM);
 
-	void InitSimulationSpace(
+	ANIMGRAPHRUNTIME_API void InitSimulationSpace(
 		const FTransform& ComponentToWorld,
 		const FTransform& BoneToComponent);
 
 	// Calculate simulation space transform, velocity etc to pass into the solver
-	void CalculateSimulationSpace(
+	ANIMGRAPHRUNTIME_API void CalculateSimulationSpace(
 		ESimulationSpace Space,
 		const FTransform& ComponentToWorld,
 		const FTransform& BoneToComponent,
@@ -342,34 +342,34 @@ private:
 		FVector& SpaceAngularAcc);
 
 	// Gather cloth collision sources from the supplied Skeltal Mesh and add a kinematic actor representing each one of them to the sim.
-	void CollectClothColliderObjects(const USkeletalMeshComponent* SkeletalMeshComp);
+	ANIMGRAPHRUNTIME_API void CollectClothColliderObjects(const USkeletalMeshComponent* SkeletalMeshComp);
 	
 	// Remove all cloth collider objects from the sim.
-	void RemoveClothColliderObjects();
+	ANIMGRAPHRUNTIME_API void RemoveClothColliderObjects();
 
 	// Update the sim-space transforms of all cloth collider objects.
-	void UpdateClothColliderObjects(const FTransform& SpaceTransform);
+	ANIMGRAPHRUNTIME_API void UpdateClothColliderObjects(const FTransform& SpaceTransform);
 
 	// Gather nearby world objects and add them to the sim
-	void CollectWorldObjects();
+	ANIMGRAPHRUNTIME_API void CollectWorldObjects();
 
 	// Flag invalid world objects to be removed from the sim
-	void ExpireWorldObjects();
+	ANIMGRAPHRUNTIME_API void ExpireWorldObjects();
 
 	// Remove simulation objects that are flagged as expired
-	void PurgeExpiredWorldObjects();
+	ANIMGRAPHRUNTIME_API void PurgeExpiredWorldObjects();
 
 	// Update sim-space transforms of world objects
-	void UpdateWorldObjects(const FTransform& SpaceTransform);
+	ANIMGRAPHRUNTIME_API void UpdateWorldObjects(const FTransform& SpaceTransform);
 
 	// Advances the simulation by a given timestep
-	void RunPhysicsSimulation(float DeltaSeconds, const FVector& SimSpaceGravity);
+	ANIMGRAPHRUNTIME_API void RunPhysicsSimulation(float DeltaSeconds, const FVector& SimSpaceGravity);
 
 	// Waits for the deferred simulation task to complete if it's not already finished
-	void FlushDeferredSimulationTask();
+	ANIMGRAPHRUNTIME_API void FlushDeferredSimulationTask();
 
 	// Destroy the simulation and free related structures
-	void DestroyPhysicsSimulation();
+	ANIMGRAPHRUNTIME_API void DestroyPhysicsSimulation();
 
 public:
 
diff --git a/Engine/Source/Runtime/AnimGraphRuntime/Public/BoneControllers/AnimNode_RotationMultiplier.h b/Engine/Source/Runtime/AnimGraphRuntime/Public/BoneControllers/AnimNode_RotationMultiplier.h
--- a/Engine/Source/Runtime/AnimGraphRuntime/Public/BoneControllers/AnimNode_RotationMultiplier.h
+++ b/Engine/Source/Runtime/AnimGraphRuntime/Public/BoneControllers/AnimNode_RotationMultiplier.h
@@ -16,7 +16,7 @@ class USkeletalMeshComponent;
  */
 
 USTRUCT(BlueprintInternalUseOnly)
-struct ANIMGRAPHRUNTIME_API FAnimNode_RotationMultiplier : public FAnimNode_SkeletalControlBase
+struct FAnimNode_RotationMultiplier : public FAnimNode_SkeletalControlBase
 {
 	GENERATED_USTRUCT_BODY()
 
@@ -39,15 +39,15 @@ struct ANIMGRAPHRUNTIME_API FAnimNode_RotationMultiplier : public FAnimNode_Skel
 	UPROPERTY(EditAnywhere, Category = Multiplier)
 	bool bIsAdditive;	
 
-	FAnimNode_RotationMultiplier();
+	ANIMGRAPHRUNTIME_API FAnimNode_RotationMultiplier();
 
 	// FAnimNode_Base interface
-	virtual void GatherDebugData(FNodeDebugData& DebugData) override;
+	ANIMGRAPHRUNTIME_API virtual void GatherDebugData(FNodeDebugData& DebugData) override;
 	// End of FAnimNode_Base interface
 
 	// FAnimNode_SkeletalControlBase interface
-	virtual void EvaluateSkeletalControl_AnyThread(FComponentSpacePoseContext& Output, TArray<FBoneTransform>& OutBoneTransforms) override;
-	virtual bool IsValidToEvaluate(const USkeleton* Skeleton, const FBoneContainer& RequiredBones) override;
+	ANIMGRAPHRUNTIME_API virtual void EvaluateSkeletalControl_AnyThread(FComponentSpacePoseContext& Output, TArray<FBoneTransform>& OutBoneTransforms) override;
+	ANIMGRAPHRUNTIME_API virtual bool IsValidToEvaluate(const USkeleton* Skeleton, const FBoneContainer& RequiredBones) override;
 	// End of FAnimNode_SkeletalControlBase interface
 
 private:
@@ -57,6 +57,6 @@ private:
 	FQuat MultiplyQuatBasedOnSourceIndex(const FTransform& RefPoseTransform, const FTransform& LocalBoneTransform, const EBoneAxis Axis, float InMultiplier, const FQuat& ReferenceQuat);
 
 	// FAnimNode_SkeletalControlBase interface
-	virtual void InitializeBoneReferences(const FBoneContainer& RequiredBones) override;
+	ANIMGRAPHRUNTIME_API virtual void InitializeBoneReferences(const FBoneContainer& RequiredBones) override;
 	// End of FAnimNode_SkeletalControlBase interface
 };
diff --git a/Engine/Source/Runtime/AnimGraphRuntime/Public/BoneControllers/AnimNode_ScaleChainLength.h b/Engine/Source/Runtime/AnimGraphRuntime/Public/BoneControllers/AnimNode_ScaleChainLength.h
--- a/Engine/Source/Runtime/AnimGraphRuntime/Public/BoneControllers/AnimNode_ScaleChainLength.h
+++ b/Engine/Source/Runtime/AnimGraphRuntime/Public/BoneControllers/AnimNode_ScaleChainLength.h
@@ -27,7 +27,7 @@ enum class EScaleChainInitialLength : uint8
  *	Scale the length of a chain of bones.
  */
 USTRUCT(BlueprintInternalUseOnly)
-struct ANIMGRAPHRUNTIME_API FAnimNode_ScaleChainLength : public FAnimNode_Base
+struct FAnimNode_ScaleChainLength : public FAnimNode_Base
 {
 	GENERATED_USTRUCT_BODY()
 
@@ -63,16 +63,16 @@ struct ANIMGRAPHRUNTIME_API FAnimNode_ScaleChainLength : public FAnimNode_Base
 	TArray<FCompactPoseBoneIndex> ChainBoneIndices;
 
 public:
-	FAnimNode_ScaleChainLength();
+	ANIMGRAPHRUNTIME_API FAnimNode_ScaleChainLength();
 
 	// FAnimNode_Base interface
-	virtual void Initialize_AnyThread(const FAnimationInitializeContext& Context) override;
-	virtual void Update_AnyThread(const FAnimationUpdateContext& Context) override;
-	virtual void CacheBones_AnyThread(const FAnimationCacheBonesContext& Context) override;
-	virtual void Evaluate_AnyThread(FPoseContext& Output) override;
-	virtual void GatherDebugData(FNodeDebugData& DebugData) override;
+	ANIMGRAPHRUNTIME_API virtual void Initialize_AnyThread(const FAnimationInitializeContext& Context) override;
+	ANIMGRAPHRUNTIME_API virtual void Update_AnyThread(const FAnimationUpdateContext& Context) override;
+	ANIMGRAPHRUNTIME_API virtual void CacheBones_AnyThread(const FAnimationCacheBonesContext& Context) override;
+	ANIMGRAPHRUNTIME_API virtual void Evaluate_AnyThread(FPoseContext& Output) override;
+	ANIMGRAPHRUNTIME_API virtual void GatherDebugData(FNodeDebugData& DebugData) override;
 	// End of FAnimNode_Base interface
 
 protected:
-	double GetInitialChainLength(FCompactPose& InLSPose, FCSPose<FCompactPose>& InCSPose) const;
+	ANIMGRAPHRUNTIME_API double GetInitialChainLength(FCompactPose& InLSPose, FCSPose<FCompactPose>& InCSPose) const;
 };
diff --git a/Engine/Source/Runtime/AnimGraphRuntime/Public/BoneControllers/AnimNode_SkeletalControlBase.h b/Engine/Source/Runtime/AnimGraphRuntime/Public/BoneControllers/AnimNode_SkeletalControlBase.h
--- a/Engine/Source/Runtime/AnimGraphRuntime/Public/BoneControllers/AnimNode_SkeletalControlBase.h
+++ b/Engine/Source/Runtime/AnimGraphRuntime/Public/BoneControllers/AnimNode_SkeletalControlBase.h
@@ -18,7 +18,7 @@
 class USkeletalMeshComponent;
 
 USTRUCT(BlueprintInternalUseOnly)
-struct ANIMGRAPHRUNTIME_API FAnimNode_SkeletalControlBase : public FAnimNode_Base
+struct FAnimNode_SkeletalControlBase : public FAnimNode_Base
 {
 	GENERATED_USTRUCT_BODY()
 
@@ -78,53 +78,53 @@ public:
 #endif //#if WITH_EDITORONLY_DATA
 
 	// FAnimNode_Base interface
-	virtual void Initialize_AnyThread(const FAnimationInitializeContext& Context) override;
-	virtual void CacheBones_AnyThread(const FAnimationCacheBonesContext& Context)  override;
-	virtual void Update_AnyThread(const FAnimationUpdateContext& Context) final;
-	virtual void EvaluateComponentSpace_AnyThread(FComponentSpacePoseContext& Output) final;
+	ANIMGRAPHRUNTIME_API virtual void Initialize_AnyThread(const FAnimationInitializeContext& Context) override;
+	ANIMGRAPHRUNTIME_API virtual void CacheBones_AnyThread(const FAnimationCacheBonesContext& Context)  override;
+	ANIMGRAPHRUNTIME_API virtual void Update_AnyThread(const FAnimationUpdateContext& Context) final;
+	ANIMGRAPHRUNTIME_API virtual void EvaluateComponentSpace_AnyThread(FComponentSpacePoseContext& Output) final;
 	virtual int32 GetLODThreshold() const override { return LODThreshold; }
 	// End of FAnimNode_Base interface
 
 	// Set the alpha of this node
-	void SetAlpha(float InAlpha);
+	ANIMGRAPHRUNTIME_API void SetAlpha(float InAlpha);
 
 	// Get the alpha of this node
-	float GetAlpha() const;
+	ANIMGRAPHRUNTIME_API float GetAlpha() const;
 
-	void InitializeAndValidateBoneRef(FBoneReference& BoneRef, const FBoneContainer& RequiredBones);
+	ANIMGRAPHRUNTIME_API void InitializeAndValidateBoneRef(FBoneReference& BoneRef, const FBoneContainer& RequiredBones);
 
 	// Visual warnings are shown on the node but not logged as an error for build system, use with care
 	// The warnigns are cleared at CacheBones_AnyThread and should be added during InitializeBoneReferences
 #if WITH_EDITOR
-	void AddBoneRefMissingVisualWarning(const FString& BoneName, const FString& SkeletalMeshName);
-	void AddValidationVisualWarning(FText ValidationVisualWarning);
-	bool HasValidationVisualWarnings() const;
-	void ClearValidationVisualWarnings();
-	FText GetValidationVisualWarningMessage() const;
+	ANIMGRAPHRUNTIME_API void AddBoneRefMissingVisualWarning(const FString& BoneName, const FString& SkeletalMeshName);
+	ANIMGRAPHRUNTIME_API void AddValidationVisualWarning(FText ValidationVisualWarning);
+	ANIMGRAPHRUNTIME_API bool HasValidationVisualWarnings() const;
+	ANIMGRAPHRUNTIME_API void ClearValidationVisualWarnings();
+	ANIMGRAPHRUNTIME_API FText GetValidationVisualWarningMessage() const;
 #endif
 
 protected:
 	// Interface for derived skeletal controls to implement
 	// use this function to update for skeletal control base
-	virtual void UpdateInternal(const FAnimationUpdateContext& Context);
+	ANIMGRAPHRUNTIME_API virtual void UpdateInternal(const FAnimationUpdateContext& Context);
 
 	// Update incoming component pose.
-	virtual void UpdateComponentPose_AnyThread(const FAnimationUpdateContext& Context);
+	ANIMGRAPHRUNTIME_API virtual void UpdateComponentPose_AnyThread(const FAnimationUpdateContext& Context);
 
 	// Evaluate incoming component pose.
-	virtual void EvaluateComponentPose_AnyThread(FComponentSpacePoseContext& Output);
+	ANIMGRAPHRUNTIME_API virtual void EvaluateComponentPose_AnyThread(FComponentSpacePoseContext& Output);
 
 	// use this function to evaluate for skeletal control base
-	virtual void EvaluateComponentSpaceInternal(FComponentSpacePoseContext& Context);
+	ANIMGRAPHRUNTIME_API virtual void EvaluateComponentSpaceInternal(FComponentSpacePoseContext& Context);
 	// Evaluate the new component-space transforms for the affected bones.
-	virtual void EvaluateSkeletalControl_AnyThread(FComponentSpacePoseContext& Output, TArray<FBoneTransform>& OutBoneTransforms);
+	ANIMGRAPHRUNTIME_API virtual void EvaluateSkeletalControl_AnyThread(FComponentSpacePoseContext& Output, TArray<FBoneTransform>& OutBoneTransforms);
 	// return true if it is valid to Evaluate
 	virtual bool IsValidToEvaluate(const USkeleton* Skeleton, const FBoneContainer& RequiredBones) { return false; }
 	// initialize any bone references you have
 	virtual void InitializeBoneReferences(const FBoneContainer& RequiredBones){};
 
 	/** Allow base to add info to the node debug output */
-	void AddDebugNodeData(FString& OutDebugData);
+	ANIMGRAPHRUNTIME_API void AddDebugNodeData(FString& OutDebugData);
 
 private:
 	// Resused bone transform array to avoid reallocating in skeletal controls
