[UE] RBAN Cloth Colliders fix for missing bodies
- Check for the the presence of collider assets in PreUpdate rather than InitPhysics as the bodies are not always present by the time InitPhysics is called. This mimics the behaviour of the cloth system.

#rnx

diff --git a/Engine/Source/Runtime/AnimGraphRuntime/Private/BoneControllers/AnimNode_RigidBody.cpp b/Engine/Source/Runtime/AnimGraphRuntime/Private/BoneControllers/AnimNode_RigidBody.cpp
--- a/Engine/Source/Runtime/AnimGraphRuntime/Private/BoneControllers/AnimNode_RigidBody.cpp
+++ b/Engine/Source/Runtime/AnimGraphRuntime/Private/BoneControllers/AnimNode_RigidBody.cpp
@@ -1005,7 +1005,9 @@ void FAnimNode_RigidBody::InitPhysics(const UAnimInstance* InAnimInstance)
 	SkeletonBoneIndexToBodyIndex.Init(INDEX_NONE, NumSkeletonBones);
 
 	PreviousTransform = SkeletalMeshComp->GetComponentToWorld();
-
+	
+	RemoveClothColliderObjects();
+	
 	ComponentsInSim.Reset();
 	ComponentsInSimTick = 0;
 
@@ -1244,8 +1246,6 @@ void FAnimNode_RigidBody::InitPhysics(const UAnimInstance* InAnimInstance)
 		PhysicsSimulation->SetIgnoreCollisionPairTable(IgnorePairs);
 		PhysicsSimulation->SetIgnoreCollisionActors(IgnoreCollisionActors);
 
-		CollectClothColliderObjects(SkeletalMeshComp);
-
 		SolverSettings = UsePhysicsAsset->SolverSettings;
 		PhysicsSimulation->SetSolverSettings(
 			SolverSettings.FixedTimeStep,
@@ -1487,6 +1487,13 @@ void FAnimNode_RigidBody::PreUpdate(const UAnimInstance* InAnimInstance)
 			}
 		}
 	}
+
+	if (bUseExternalClothCollision && ClothColliders.IsEmpty())
+	{
+		// The Cloth Collider assets are part of the SkelMeshComponent and can be initialized after the first call to InitPhysics. Keep checking here until some 
+		// are found, following the behavior of the cloth system (see USkeletalMeshComponent::UpdateClothTransformImp())
+		CollectClothColliderObjects(SKC);
+	}
 }
 
 int32 FAnimNode_RigidBody::GetLODThreshold() const
