[UE] Call FActorHandle::SetKinematicTarget instead of FActorHandle::SetWorldTransform in FAnimNode_RigidBody
- Fixes bug (introduced in CL-25315823) where Kinematic bodies included in an RBAN simuation to represent cloth colliders or world space objects fail to collide with other bodies because their bounding boxes are not updated after calls to FActorHandle::SetWorldTransform.

#rnx

diff --git a/Engine/Source/Runtime/AnimGraphRuntime/Private/BoneControllers/AnimNode_RigidBody.cpp b/Engine/Source/Runtime/AnimGraphRuntime/Private/BoneControllers/AnimNode_RigidBody.cpp
--- a/Engine/Source/Runtime/AnimGraphRuntime/Private/BoneControllers/AnimNode_RigidBody.cpp
+++ b/Engine/Source/Runtime/AnimGraphRuntime/Private/BoneControllers/AnimNode_RigidBody.cpp
@@ -1611,7 +1611,7 @@ void FAnimNode_RigidBody::UpdateClothColliderObjects(const FTransform& SpaceTran
 			CompSpaceTransform.SetScale3D(FVector::OneVector);	// TODO - sort out scale for world objects in local sim
 
 			// Update the sim's copy of the world object
-			ClothCollider.ActorHandle->SetWorldTransform(CompSpaceTransform);
+			ClothCollider.ActorHandle->SetKinematicTarget(CompSpaceTransform);
 		}
 	}
 }
@@ -1725,7 +1725,7 @@ void FAnimNode_RigidBody::UpdateWorldObjects(const FTransform& SpaceTransform)
 				CompSpaceTransform.SetScale3D(FVector::OneVector);	// TODO - sort out scale for world objects in local sim
 
 				// Update the sim's copy of the world object
-				ActorHandle->SetWorldTransform(CompSpaceTransform);
+				ActorHandle->SetKinematicTarget(CompSpaceTransform);
 			}
 		}
 	}
