Allow users to cast a Anim Node Reference to a Mirror Node inside a Anim Node Function.

- Allow the user to set/get Mirror flag, and BlendTime using a Mirror Anim Node Reference.

diff --git a/Engine/Source/Runtime/AnimGraphRuntime/Private/MirrorAnimLibrary.cpp b/Engine/Source/Runtime/AnimGraphRuntime/Private/MirrorAnimLibrary.cpp
new file mode 100644
--- /dev/null
+++ b/Engine/Source/Runtime/AnimGraphRuntime/Private/MirrorAnimLibrary.cpp
@@ -0,0 +1,74 @@
+﻿// Copyright Epic Games, Inc. All Rights Reserved.
+
+#include "MirrorAnimLibrary.h"
+#include "AnimNodes/AnimNode_Mirror.h"
+#include UE_INLINE_GENERATED_CPP_BY_NAME(MirrorAnimLibrary)
+
+DEFINE_LOG_CATEGORY_STATIC(LogMirrorNodeLibrary, Verbose, All);
+
+FMirrorAnimNodeReference UMirrorAnimLibrary::ConvertToMirrorNode(const FAnimNodeReference& Node, EAnimNodeReferenceConversionResult& Result)
+{
+	return FAnimNodeReference::ConvertToType<FMirrorAnimNodeReference>(Node, Result);
+}
+
+FMirrorAnimNodeReference UMirrorAnimLibrary::SetMirror(const FMirrorAnimNodeReference& MirrorNode, bool bInMirror)
+{
+	MirrorNode.CallAnimNodeFunction<FAnimNode_MirrorBase>(
+		TEXT("SetMirror"),[bInMirror](FAnimNode_MirrorBase& InMirrorBaseNode)
+		{
+			if (!InMirrorBaseNode.SetMirror(bInMirror))
+			{
+				UE_LOG(LogMirrorNodeLibrary, Warning, TEXT("Could not set Mirror flag on mirror anim node, value is not dynamic. Set it as Always Dynamic."));
+			}
+		});
+	return MirrorNode;
+}
+
+FMirrorAnimNodeReference UMirrorAnimLibrary::SetMirrorTransitionBlendTime(const FMirrorAnimNodeReference& MirrorNode, float InBlendTime)
+{
+	MirrorNode.CallAnimNodeFunction<FAnimNode_MirrorBase>(
+		TEXT("SetBlendTimeOnMirrorStateChange"),[InBlendTime](FAnimNode_MirrorBase& InMirrorBaseNode)
+		{
+			if (!InMirrorBaseNode.SetBlendTimeOnMirrorStateChange(InBlendTime))
+			{
+				UE_LOG(LogMirrorNodeLibrary, Warning, TEXT("Could not set Mirror Transition Blend Time on mirror anim node, value is not dynamic. Set it as Always Dynamic."));
+			}
+		});
+	return MirrorNode;
+}
+
+bool UMirrorAnimLibrary::GetMirror(const FMirrorAnimNodeReference& MirrorNode)
+{
+	bool bIsMirrored = false;
+	
+	MirrorNode.CallAnimNodeFunction<FAnimNode_MirrorBase>(
+		TEXT("GetMirror"),[&bIsMirrored](FAnimNode_MirrorBase& InMirrorBaseNode)
+		{
+			bIsMirrored = InMirrorBaseNode.GetMirror();
+		});
+	return bIsMirrored;
+}
+
+UMirrorDataTable* UMirrorAnimLibrary::GetMirrorDataTable(const FMirrorAnimNodeReference& MirrorNode)
+{
+	UMirrorDataTable* MirrorDataTable = nullptr;
+	
+	MirrorNode.CallAnimNodeFunction<FAnimNode_MirrorBase>(
+		TEXT("GetMirrorDataTable"),[&MirrorDataTable](FAnimNode_MirrorBase& InMirrorBaseNode)
+		{
+			MirrorDataTable = InMirrorBaseNode.GetMirrorDataTable();
+		});
+	return MirrorDataTable;
+}
+
+float UMirrorAnimLibrary::GetMirrorTransitionBlendTime(const FMirrorAnimNodeReference& MirrorNode)
+{
+	float BlendTime = false;
+	
+	MirrorNode.CallAnimNodeFunction<FAnimNode_MirrorBase>(
+		TEXT("GetBlendTimeOnMirrorStateChange"),[&BlendTime](FAnimNode_MirrorBase& InMirrorBaseNode)
+		{
+			BlendTime = InMirrorBaseNode.GetBlendTimeOnMirrorStateChange();
+		});
+	return BlendTime;
+}
\ No newline at end of file
diff --git a/Engine/Source/Runtime/AnimGraphRuntime/Public/MirrorAnimLibrary.h b/Engine/Source/Runtime/AnimGraphRuntime/Public/MirrorAnimLibrary.h
new file mode 100644
--- /dev/null
+++ b/Engine/Source/Runtime/AnimGraphRuntime/Public/MirrorAnimLibrary.h
@@ -0,0 +1,59 @@
+﻿// Copyright Epic Games, Inc. All Rights Reserved.
+#pragma once
+
+#include "CoreMinimal.h"
+
+#include "Animation/AnimExecutionContext.h"
+#include "Kismet/BlueprintFunctionLibrary.h"
+#include "Animation/AnimNodeReference.h"
+#include "MirrorAnimLibrary.generated.h"
+
+struct FAnimNode_MirrorBase;
+
+USTRUCT(BlueprintType)
+struct FMirrorAnimNodeReference : public FAnimNodeReference
+{
+	GENERATED_BODY()
+	typedef FAnimNode_MirrorBase FInternalNodeType;
+};
+
+/** Exposes operations that can be run on a Mirror node via Anim Node Functions such as "On Become Relevant" and "On Update". */
+UCLASS()
+class UMirrorAnimLibrary : public UBlueprintFunctionLibrary
+{
+	GENERATED_BODY()
+public:
+	
+	/** Get a mirror node context from an anim node context */
+	UFUNCTION(BlueprintCallable, Category = "Animation|Mirroring", meta=(BlueprintThreadSafe, ExpandEnumAsExecs = "Result"))
+	static ANIMGRAPHRUNTIME_API FMirrorAnimNodeReference ConvertToMirrorNode(const FAnimNodeReference& Node, EAnimNodeReferenceConversionResult& Result);
+	
+	/** Get a mirror context from an anim node context (pure) */
+	UFUNCTION(BlueprintPure, Category = "Animation|Mirroring", meta=(BlueprintThreadSafe, DisplayName = "Convert to Mirror Node"))
+	static void ConvertToMirrorNodePure(const FAnimNodeReference& Node, FMirrorAnimNodeReference& MirrorNode, bool& Result)
+	{
+		EAnimNodeReferenceConversionResult ConversionResult;
+		MirrorNode = ConvertToMirrorNode(Node, ConversionResult);
+		Result = (ConversionResult == EAnimNodeReferenceConversionResult::Succeeded);
+	}
+	
+	/** Set the mirror state */
+	UFUNCTION(BlueprintCallable, Category = "Animation|Mirroring", meta=(BlueprintThreadSafe))
+	static ANIMGRAPHRUNTIME_API FMirrorAnimNodeReference SetMirror(const FMirrorAnimNodeReference& MirrorNode, bool bInMirror);
+	
+	/** Set how long to blend using inertialization when switching mirrored state */
+	UFUNCTION(BlueprintCallable, Category = "Animation|Mirroring", meta=(BlueprintThreadSafe))
+	static ANIMGRAPHRUNTIME_API FMirrorAnimNodeReference SetMirrorTransitionBlendTime(const FMirrorAnimNodeReference& MirrorNode, float InBlendTime);
+	
+	/** Get the mirror state */
+	UFUNCTION(BlueprintPure, Category = "Animation|Mirroring", meta=(BlueprintThreadSafe))
+	static ANIMGRAPHRUNTIME_API bool GetMirror(const FMirrorAnimNodeReference& MirrorNode);
+	
+	/** Get MirrorDataTable used to perform mirroring*/
+	UFUNCTION(BlueprintPure, Category = "Animation|Mirroring", meta=(BlueprintThreadSafe))
+	static ANIMGRAPHRUNTIME_API UMirrorDataTable* GetMirrorDataTable(const FMirrorAnimNodeReference& MirrorNode);
+	
+	/** Get how long to blend using inertialization when switching mirrored state */
+	UFUNCTION(BlueprintPure, Category = "Animation|Mirroring", meta=(BlueprintThreadSafe))
+	static ANIMGRAPHRUNTIME_API float GetMirrorTransitionBlendTime(const FMirrorAnimNodeReference& MirrorNode);
+};
\ No newline at end of file
