Allow users to cast a Anim Node Reference to a Modify Curve Node.

- Allow the user to set/get Alpha.
- Allow the user to set the Curve Map.

diff --git a/Engine/Source/Runtime/AnimGraphRuntime/Private/ModifyCurveLibrary.cpp b/Engine/Source/Runtime/AnimGraphRuntime/Private/ModifyCurveLibrary.cpp
new file mode 100644
--- /dev/null
+++ b/Engine/Source/Runtime/AnimGraphRuntime/Private/ModifyCurveLibrary.cpp
@@ -0,0 +1,66 @@
+﻿// Copyright Epic Games, Inc. All Rights Reserved.
+
+#include "ModifyCurveLibrary.h"
+#include "AnimNodes/AnimNode_ModifyCurve.h"
+#include UE_INLINE_GENERATED_CPP_BY_NAME(ModifyCurveLibrary)
+
+DEFINE_LOG_CATEGORY_STATIC(LogModifyCurveNodeLibrary, Verbose, All);
+
+FModifyCurveAnimNodeReference UModifyCurveAnimLibrary::ConvertToModifyCurveNode(const FAnimNodeReference& Node, EAnimNodeReferenceConversionResult& Result)
+{
+	return FAnimNodeReference::ConvertToType<FModifyCurveAnimNodeReference>(Node, Result);
+}
+
+FModifyCurveAnimNodeReference UModifyCurveAnimLibrary::SetCurveMap(const FModifyCurveAnimNodeReference& ModifyCurveNode, const TMap<FName, float> & InCurveMap)
+{
+	ModifyCurveNode.CallAnimNodeFunction<FAnimNode_ModifyCurve>(TEXT("SetCurveMap"),[&InCurveMap](FAnimNode_ModifyCurve& InModifyCurveNode)
+	{
+		InModifyCurveNode.CurveMap.Append(InCurveMap);
+	});
+	
+	return ModifyCurveNode;
+}
+
+EModifyCurveApplyMode UModifyCurveAnimLibrary::GetApplyMode(const FModifyCurveAnimNodeReference& ModifyCurveNode)
+{
+	EModifyCurveApplyMode OutMode = EModifyCurveApplyMode::Blend;
+	
+	ModifyCurveNode.CallAnimNodeFunction<FAnimNode_ModifyCurve>(TEXT("GetApplyMode"),[&OutMode](FAnimNode_ModifyCurve& InModifyCurveNode)
+	{
+		OutMode = InModifyCurveNode.ApplyMode;
+	});
+		
+	return OutMode;
+}
+
+FModifyCurveAnimNodeReference UModifyCurveAnimLibrary::SetApplyMode(const FModifyCurveAnimNodeReference& ModifyCurveNode, EModifyCurveApplyMode InMode)
+{
+	ModifyCurveNode.CallAnimNodeFunction<FAnimNode_ModifyCurve>(TEXT("SetApplyMode"),[InMode](FAnimNode_ModifyCurve& InModifyCurveNode)
+	{
+		InModifyCurveNode.ApplyMode = InMode;
+	});
+	
+	return ModifyCurveNode;
+}
+
+float UModifyCurveAnimLibrary::GetAlpha(const FModifyCurveAnimNodeReference& ModifyCurveNode)
+{
+	float OutAlpha = 0.0f;
+
+	ModifyCurveNode.CallAnimNodeFunction<FAnimNode_ModifyCurve>(TEXT("GetAlpha"), [&OutAlpha](const FAnimNode_ModifyCurve& InModifyCurveNode)
+	{
+		OutAlpha = InModifyCurveNode.Alpha;
+	});
+	
+	return OutAlpha;
+}
+
+FModifyCurveAnimNodeReference UModifyCurveAnimLibrary::SetAlpha(const FModifyCurveAnimNodeReference& ModifyCurveNode, float InAlpha)
+{
+	ModifyCurveNode.CallAnimNodeFunction<FAnimNode_ModifyCurve>(TEXT("SetAlpha"), [InAlpha](FAnimNode_ModifyCurve& InModifyCurveNode)
+	{
+		InModifyCurveNode.Alpha = InAlpha;
+	});
+	
+	return ModifyCurveNode;
+}
diff --git a/Engine/Source/Runtime/AnimGraphRuntime/Public/ModifyCurveLibrary.h b/Engine/Source/Runtime/AnimGraphRuntime/Public/ModifyCurveLibrary.h
new file mode 100644
--- /dev/null
+++ b/Engine/Source/Runtime/AnimGraphRuntime/Public/ModifyCurveLibrary.h
@@ -0,0 +1,54 @@
+﻿// Copyright Epic Games, Inc. All Rights Reserved.
+#pragma once
+
+#include "CoreMinimal.h"
+
+#include "Animation/AnimExecutionContext.h"
+#include "Kismet/BlueprintFunctionLibrary.h"
+#include "Animation/AnimNodeReference.h"
+#include "ModifyCurveLibrary.generated.h"
+
+struct FAnimNode_ModifyCurve;
+
+USTRUCT(BlueprintType)
+struct FModifyCurveAnimNodeReference : public FAnimNodeReference
+{
+	GENERATED_BODY()
+	typedef FAnimNode_ModifyCurve FInternalNodeType;
+};
+
+/** Exposes operations that can be run on a Modify Curve Node via Anim Node Functions such as "On Become Relevant" and "On Update". */
+UCLASS()
+class UModifyCurveAnimLibrary : public UBlueprintFunctionLibrary
+{
+	GENERATED_BODY()
+public:
+	
+	/** Get a modify curve node context from an anim node context */
+	UFUNCTION(BlueprintCallable, Category = "Animation", meta=(BlueprintThreadSafe, ExpandEnumAsExecs = "Result"))
+	static ANIMGRAPHRUNTIME_API FModifyCurveAnimNodeReference ConvertToModifyCurveNode(const FAnimNodeReference& Node, EAnimNodeReferenceConversionResult& Result);
+	
+	/** Get a modify curve context from an anim node context (pure) */
+	UFUNCTION(BlueprintPure, Category = "Animation", meta=(BlueprintThreadSafe, DisplayName = "Convert to Modify Curve node"))
+	static void ConvertToModifyCurveNodePure(const FAnimNodeReference& Node, FModifyCurveAnimNodeReference& ModifyCurveNode, bool& Result)
+	{
+		EAnimNodeReferenceConversionResult ConversionResult;
+		ModifyCurveNode = ConvertToModifyCurveNode(Node, ConversionResult);
+		Result = (ConversionResult == EAnimNodeReferenceConversionResult::Succeeded);
+	}
+
+	UFUNCTION(BlueprintCallable, Category = "Animation", meta=(BlueprintThreadSafe))
+	static ANIMGRAPHRUNTIME_API FModifyCurveAnimNodeReference SetCurveMap(const FModifyCurveAnimNodeReference& ModifyCurveNode, const TMap<FName, float> & InCurveMap);
+
+	UFUNCTION(BlueprintCallable, Category = "Animation", meta=(BlueprintThreadSafe))
+	static ANIMGRAPHRUNTIME_API EModifyCurveApplyMode GetApplyMode(const FModifyCurveAnimNodeReference& ModifyCurveNode);
+	
+	UFUNCTION(BlueprintCallable, Category = "Animation", meta=(BlueprintThreadSafe))
+	static ANIMGRAPHRUNTIME_API FModifyCurveAnimNodeReference SetApplyMode(const FModifyCurveAnimNodeReference& ModifyCurveNode, EModifyCurveApplyMode InMode);
+	
+	UFUNCTION(BlueprintCallable, Category = "Animation", meta=(BlueprintThreadSafe))
+	static ANIMGRAPHRUNTIME_API float GetAlpha(const FModifyCurveAnimNodeReference& ModifyCurveNode);
+	
+	UFUNCTION(BlueprintCallable, Category = "Animation", meta=(BlueprintThreadSafe))
+	static ANIMGRAPHRUNTIME_API FModifyCurveAnimNodeReference SetAlpha(const FModifyCurveAnimNodeReference& ModifyCurveNode, float InAlpha);
+};
\ No newline at end of file
