Fixes application of smoothing in the lookat aim offset blendspace

diff --git a/Engine/Source/Runtime/AnimGraphRuntime/Private/AnimNodes/AnimNode_AimOffsetLookAt.cpp b/Engine/Source/Runtime/AnimGraphRuntime/Private/AnimNodes/AnimNode_AimOffsetLookAt.cpp
--- a/Engine/Source/Runtime/AnimGraphRuntime/Private/AnimNodes/AnimNode_AimOffsetLookAt.cpp
+++ b/Engine/Source/Runtime/AnimGraphRuntime/Private/AnimNodes/AnimNode_AimOffsetLookAt.cpp
@@ -172,11 +172,15 @@ void FAnimNode_AimOffsetLookAt::UpdateFromLookAtTarget(FPoseContext& LocalPoseCo
 #endif // ENABLE_DRAW_DEBUG
 	}
 
-	// Update Blend Space and put the result into BlendSampleDataCache.
+	// Update Blend Space, including the smoothing/filtering, and put the result into BlendSampleDataCache.
 	if (CurrentBlendSpace)
 	{
+		const FVector BlendSpacePosition(CurrentBlendInput.X, CurrentBlendInput.Y, 0.f);
+		const FVector FilteredBlendInput = CurrentBlendSpace->FilterInput(
+			&BlendFilter, BlendSpacePosition, DeltaTimeRecord.Delta);
+
 		CurrentBlendSpace->UpdateBlendSamples(
-			CurrentBlendInput, DeltaTimeRecord.Delta, BlendSampleDataCache, CachedTriangulationIndex);
+			FilteredBlendInput, DeltaTimeRecord.Delta, BlendSampleDataCache, CachedTriangulationIndex);
 	}
 }
 
