Expose anim graph IsActive to blueprint

diff --git a/Engine/Source/Runtime/AnimGraphRuntime/Private/AnimExecutionContextLibrary.cpp b/Engine/Source/Runtime/AnimGraphRuntime/Private/AnimExecutionContextLibrary.cpp
--- a/Engine/Source/Runtime/AnimGraphRuntime/Private/AnimExecutionContextLibrary.cpp
+++ b/Engine/Source/Runtime/AnimGraphRuntime/Private/AnimExecutionContextLibrary.cpp
@@ -54,6 +54,13 @@ float UAnimExecutionContextLibrary::GetCurrentWeight(const FAnimUpdateContext& U
 	return 0.0f;
 }
 
+bool UAnimExecutionContextLibrary::IsActive(const FAnimExecutionContext& Context)
+{
+	FAnimationBaseContext* BaseContext = Context.GetBaseContext();
+	check(BaseContext);
+	return BaseContext->IsActive();
+}
+
 FAnimPoseContext UAnimExecutionContextLibrary::ConvertToPoseContext(const FAnimExecutionContext& Context, EAnimExecutionContextConversionResult& Result)
 {
 	return FAnimExecutionContext::ConvertToType<FAnimPoseContext>(Context, Result);
diff --git a/Engine/Source/Runtime/AnimGraphRuntime/Public/AnimExecutionContextLibrary.h b/Engine/Source/Runtime/AnimGraphRuntime/Public/AnimExecutionContextLibrary.h
--- a/Engine/Source/Runtime/AnimGraphRuntime/Public/AnimExecutionContextLibrary.h
+++ b/Engine/Source/Runtime/AnimGraphRuntime/Public/AnimExecutionContextLibrary.h
@@ -51,6 +51,10 @@ public:
 	UFUNCTION(BlueprintPure, Category = "Animation|Utilities", meta=(BlueprintThreadSafe))
 	static ANIMGRAPHRUNTIME_API float GetCurrentWeight(const FAnimUpdateContext& Context);
 	
+	/** Get whether this branch of the graph is active (i.e. NOT blending out). */
+	UFUNCTION(BlueprintPure, Category = "Animation|Utilities", meta=(BlueprintThreadSafe))
+	static ANIMGRAPHRUNTIME_API bool IsActive(const FAnimExecutionContext& Context);
+
 	/** Convert to a pose context */
 	UFUNCTION(BlueprintCallable, Category = "Animation|Utilities", meta=(BlueprintThreadSafe, ExpandEnumAsExecs = "Result"))
 	static ANIMGRAPHRUNTIME_API FAnimPoseContext ConvertToPoseContext(const FAnimExecutionContext& Context, EAnimExecutionContextConversionResult& Result);
