Add PostEvaluateSourcePose interface to FAnimNode_Slot to be able to override some of the curves in derived nodes

diff --git a/Engine/Source/Runtime/AnimGraphRuntime/Private/AnimNodes/AnimNode_Slot.cpp b/Engine/Source/Runtime/AnimGraphRuntime/Private/AnimNodes/AnimNode_Slot.cpp
--- a/Engine/Source/Runtime/AnimGraphRuntime/Private/AnimNodes/AnimNode_Slot.cpp
+++ b/Engine/Source/Runtime/AnimGraphRuntime/Private/AnimNodes/AnimNode_Slot.cpp
@@ -91,6 +91,7 @@ void FAnimNode_Slot::Evaluate_AnyThread(FPoseContext & Output)
 	if (WeightData.SlotNodeWeight <= ZERO_ANIMWEIGHT_THRESH)
 	{
 		Source.Evaluate(Output);
+		PostEvaluateSourcePose(Output);
 	}
 	else
 	{
@@ -100,6 +101,7 @@ void FAnimNode_Slot::Evaluate_AnyThread(FPoseContext & Output)
 			Source.Evaluate(SourceContext);
 		}
 
+		PostEvaluateSourcePose(SourceContext);
 		const FAnimationPoseData SourcePoseData(SourceContext);
 		FAnimationPoseData OutputPoseData(Output);
 		Output.AnimInstanceProxy->SlotEvaluatePose(SlotName, SourcePoseData, WeightData.SourceWeight, OutputPoseData, WeightData.SlotNodeWeight, WeightData.TotalNodeWeight);
diff --git a/Engine/Source/Runtime/AnimGraphRuntime/Public/AnimNodes/AnimNode_Slot.h b/Engine/Source/Runtime/AnimGraphRuntime/Public/AnimNodes/AnimNode_Slot.h
--- a/Engine/Source/Runtime/AnimGraphRuntime/Public/AnimNodes/AnimNode_Slot.h
+++ b/Engine/Source/Runtime/AnimGraphRuntime/Public/AnimNodes/AnimNode_Slot.h
@@ -29,6 +29,8 @@ struct FAnimNode_Slot : public FAnimNode_Base
 	bool bAlwaysUpdateSourcePose;
 
 protected:
+	virtual void PostEvaluateSourcePose(FPoseContext& SourceContext) {}
+
 	FSlotNodeWeightInfo WeightData;
 	FGraphTraversalCounter SlotNodeInitializationCounter;
 
