Fixed inertialization not working properly with sync groups when using BlendList.

[RN] minor

diff --git a/Engine/Source/Runtime/AnimGraphRuntime/Private/AnimNodes/AnimNode_BlendListBase.cpp b/Engine/Source/Runtime/AnimGraphRuntime/Private/AnimNodes/AnimNode_BlendListBase.cpp
--- a/Engine/Source/Runtime/AnimGraphRuntime/Private/AnimNodes/AnimNode_BlendListBase.cpp
+++ b/Engine/Source/Runtime/AnimGraphRuntime/Private/AnimNodes/AnimNode_BlendListBase.cpp
@@ -2,6 +2,7 @@
 
 #include "AnimNodes/AnimNode_BlendListBase.h"
 #include "AnimationRuntime.h"
+#include "Animation/AnimInertializationSyncScope.h"
 #include "Animation/BlendProfile.h"
 #include "Animation/AnimInstanceProxy.h"
 #include "Animation/AnimNode_Inertialization.h"
@@ -96,7 +97,8 @@ void FAnimNode_BlendListBase::Update_AnyThread(const FAnimationUpdateContext& Co
 	const int32 NumPoses = BlendPose.Num();
 	const TArray<float>& CurrentBlendTimes = GetBlendTimes();
 	checkSlow(PerBlendData.Num() == NumPoses);
-
+	bool bRequestedInertializationOnActiveChildIndexChange = false;
+	
 	if (NumPoses > 0)
 	{
 		UBlendProfile* CurrentBlendProfile = GetBlendProfile();
@@ -138,6 +140,7 @@ void FAnimNode_BlendListBase::Update_AnyThread(const FAnimationUpdateContext& Co
 
 					InertializationRequester->RequestInertialization(Request);
 					InertializationRequester->AddDebugRecord(*Context.AnimInstanceProxy, Context.GetCurrentNodeId());
+					bRequestedInertializationOnActiveChildIndexChange = true;
 				}
 				else
 				{
@@ -232,6 +235,8 @@ void FAnimNode_BlendListBase::Update_AnyThread(const FAnimationUpdateContext& Co
 			if (BlendWeight > ZERO_ANIMWEIGHT_THRESH)
 			{
 				FAnimationUpdateContext ChildContext = Context.FractionalWeight(BlendWeight);
+				
+				UE::Anim::TOptionalScopedGraphMessage<UE::Anim::FAnimInertializationSyncScope> InertializationSync(bRequestedInertializationOnActiveChildIndexChange, ChildContext);
 				BlendPose[i].Update((i == ChildIndex) ? ChildContext : ChildContext.AsInactive());
 			}
 		}
diff --git a/Engine/Source/Runtime/Engine/Public/Animation/AnimInertializationSyncScope.h b/Engine/Source/Runtime/Engine/Public/Animation/AnimInertializationSyncScope.h
--- a/Engine/Source/Runtime/Engine/Public/Animation/AnimInertializationSyncScope.h
+++ b/Engine/Source/Runtime/Engine/Public/Animation/AnimInertializationSyncScope.h
@@ -10,7 +10,7 @@ namespace UE { namespace Anim {
 // Scoped graph message used to synchronize tick records using inertialization 
 class FAnimInertializationSyncScope : public IGraphMessage
 {
-	DECLARE_ANIMGRAPH_MESSAGE(FAnimInertializationSyncScope)
+	DECLARE_ANIMGRAPH_MESSAGE_API(FAnimInertializationSyncScope, ENGINE_API)
 };
 
 }} // namespace UE::Anim
