Need the ability to give a sequence evaluator a frame instead of time.
#feature sequence evaluator node now exposes ExplicitFrame option, which internally combined with the set AnimSequence is used to generate a time-value (on which the behaviour is based)

diff --git a/Engine/Source/Editor/AnimGraph/Private/AnimGraphNode_SequenceEvaluator.cpp b/Engine/Source/Editor/AnimGraph/Private/AnimGraphNode_SequenceEvaluator.cpp
--- a/Engine/Source/Editor/AnimGraph/Private/AnimGraphNode_SequenceEvaluator.cpp
+++ b/Engine/Source/Editor/AnimGraph/Private/AnimGraphNode_SequenceEvaluator.cpp
@@ -9,6 +9,7 @@
 #include "Animation/AnimSequence.h"
 #include "AnimGraphNode_SequenceEvaluator.h"
 
+#include "AnimGraphNodeBinding.h"
 #include "BlueprintNodeSpawner.h"
 #include "EditorCategoryUtils.h"
 #include "IAnimBlueprintNodeOverrideAssetsContext.h"
@@ -61,7 +62,7 @@ FText UAnimGraphNode_SequenceEvaluator::GetMenuCategory() const
 FText UAnimGraphNode_SequenceEvaluator::GetNodeTitle(ENodeTitleType::Type TitleType) const
 {
 	UEdGraphPin* SequencePin = FindPin(GET_MEMBER_NAME_STRING_CHECKED(FAnimNode_SequenceEvaluator, Sequence));
-	return GetNodeTitleHelper(TitleType, SequencePin, LOCTEXT("PlayerDesc", "Sequence Evaluator"));
+	return GetNodeTitleHelper(TitleType, SequencePin, Node.ShouldUseExplicitFrame() ? LOCTEXT("PlayerDescByFrame", "Sequence Evaluator (by frame)") :  LOCTEXT("PlayerDescByTime", "Sequence Evaluator (by time)"));
 }
 
 FSlateIcon UAnimGraphNode_SequenceEvaluator::GetIconAndTint(FLinearColor& OutColor) const
@@ -136,6 +137,37 @@ void UAnimGraphNode_SequenceEvaluator::CopySettingsFromAnimationAsset(UAnimation
 	}
 }
 
+void UAnimGraphNode_SequenceEvaluator::PostEditChangeProperty(FPropertyChangedEvent& PropertyChangedEvent)
+{
+	Super::PostEditChangeProperty(PropertyChangedEvent);
+
+	if( PropertyChangedEvent.GetPropertyName() == GET_MEMBER_NAME_CHECKED(FAnimNode_SequenceEvaluator, bUseExplicitFrame))
+	{
+		UEdGraphPin** Pin = nullptr;
+		if (Node.bUseExplicitFrame)
+		{
+			Pin = Pins.FindByPredicate([](const UEdGraphPin* InPin)
+			{
+				return InPin->PinName == GET_MEMBER_NAME_STRING_CHECKED(FAnimNode_SequenceEvaluator, ExplicitTime);
+			});
+		}
+		else
+		{
+			Pin = Pins.FindByPredicate([](const UEdGraphPin* InPin)
+			{
+				return InPin->PinName == GET_MEMBER_NAME_STRING_CHECKED(FAnimNode_SequenceEvaluator, ExplicitFrame);
+			});			
+		}
+
+		if (Pin)
+		{			
+			Modify();
+			(*Pin)->BreakAllPinLinks();
+			RemoveBindings((*Pin)->PinName);
+		}
+	}
+}
+
 void UAnimGraphNode_SequenceEvaluator::OnOverrideAssets(IAnimBlueprintNodeOverrideAssetsContext& InContext) const
 {
 	if(InContext.GetAssets().Num() > 0)
diff --git a/Engine/Source/Editor/AnimGraph/Public/AnimGraphNode_SequenceEvaluator.h b/Engine/Source/Editor/AnimGraph/Public/AnimGraphNode_SequenceEvaluator.h
--- a/Engine/Source/Editor/AnimGraph/Public/AnimGraphNode_SequenceEvaluator.h
+++ b/Engine/Source/Editor/AnimGraph/Public/AnimGraphNode_SequenceEvaluator.h
@@ -50,4 +50,6 @@ class UAnimGraphNode_SequenceEvaluator : public UAnimGraphNode_AssetPlayerBase
 	virtual void SetAnimationAsset(UAnimationAsset* Asset) override;
 	virtual void CopySettingsFromAnimationAsset(UAnimationAsset* Asset) override;
 	// End of UAnimGraphNode_AssetPlayerBase interface
+	
+	virtual void PostEditChangeProperty(FPropertyChangedEvent& PropertyChangedEvent) override;
 };
diff --git a/Engine/Source/Runtime/AnimGraphRuntime/Private/AnimNodes/AnimNode_SequenceEvaluator.cpp b/Engine/Source/Runtime/AnimGraphRuntime/Private/AnimNodes/AnimNode_SequenceEvaluator.cpp
--- a/Engine/Source/Runtime/AnimGraphRuntime/Private/AnimNodes/AnimNode_SequenceEvaluator.cpp
+++ b/Engine/Source/Runtime/AnimGraphRuntime/Private/AnimNodes/AnimNode_SequenceEvaluator.cpp
@@ -152,9 +152,60 @@ UAnimSequenceBase* FAnimNode_SequenceEvaluator::GetSequence() const
 	return GET_ANIM_NODE_DATA(TObjectPtr<UAnimSequenceBase>, Sequence);
 }
 
+bool FAnimNode_SequenceEvaluator::SetShouldUseExplicitFrame(bool bFlag)
+{	
+#if WITH_EDITORONLY_DATA
+	bUseExplicitFrame = bFlag;
+#endif
+
+	if (bool* bUseExplicitFrameState = GET_INSTANCE_ANIM_NODE_DATA_PTR(bool, bUseExplicitFrame))
+	{
+		*bUseExplicitFrameState = bFlag;
+		return true;
+	}
+
+	return false;
+}
+
+bool FAnimNode_SequenceEvaluator::ShouldUseExplicitFrame() const
+{
+	return GET_ANIM_NODE_DATA(bool, bUseExplicitFrame);
+}
+
+int32 FAnimNode_SequenceEvaluator::GetExplicitFrame() const
+{
+	return GET_ANIM_NODE_DATA(int32, ExplicitFrame);
+}
+
+bool FAnimNode_SequenceEvaluator::SetExplicitFrame(int32 InFrame)
+{	
+#if WITH_EDITORONLY_DATA
+	ExplicitFrame = InFrame;
+#endif
+
+	if (int32* ExplicitFramePtr = GET_INSTANCE_ANIM_NODE_DATA_PTR(int32, ExplicitFrame))
+	{
+		*ExplicitFramePtr = InFrame;
+		return true;
+	}
+
+	return false;
+}
+
 float FAnimNode_SequenceEvaluator::GetExplicitTime() const
 {
-	return GET_ANIM_NODE_DATA(float, ExplicitTime);
+	if(ShouldUseExplicitFrame())
+	{
+		if(const UAnimSequenceBase* SequenceBase = GetSequence())
+		{
+			const int32 Frame = GET_ANIM_NODE_DATA(int32, ExplicitFrame);
+			return SequenceBase->GetSamplingFrameRate().AsSeconds(Frame);
+		}
+
+		return 0.f;
+	}
+	  
+	 return GET_ANIM_NODE_DATA(float, ExplicitTime);
 }
 
 bool FAnimNode_SequenceEvaluator::SetExplicitTime(float InTime)
@@ -285,4 +336,4 @@ bool FAnimNode_SequenceEvaluator::SetIgnoreForRelevancyTest(bool bInIgnoreForRel
 	}
 
 	return false;
-}
+}
\ No newline at end of file
diff --git a/Engine/Source/Runtime/AnimGraphRuntime/Private/SequenceEvaluatorLibrary.cpp b/Engine/Source/Runtime/AnimGraphRuntime/Private/SequenceEvaluatorLibrary.cpp
--- a/Engine/Source/Runtime/AnimGraphRuntime/Private/SequenceEvaluatorLibrary.cpp
+++ b/Engine/Source/Runtime/AnimGraphRuntime/Private/SequenceEvaluatorLibrary.cpp
@@ -32,6 +32,21 @@ FSequenceEvaluatorReference USequenceEvaluatorLibrary::SetExplicitTime(const FSe
 	return SequenceEvaluator;
 }
 
+FSequenceEvaluatorReference USequenceEvaluatorLibrary::SetExplicitFrame(const FSequenceEvaluatorReference& SequenceEvaluator, int32 Frame)
+{
+	SequenceEvaluator.CallAnimNodeFunction<FAnimNode_SequenceEvaluator>(
+		TEXT("SetExplicitTime"),
+		[Frame](FAnimNode_SequenceEvaluator& InSequenceEvaluator)
+		{
+			if(!InSequenceEvaluator.SetExplicitFrame(Frame))
+			{
+				UE_LOG(LogSequenceEvaluatorLibrary, Warning, TEXT("Could not set explicit frame on sequence evaluator, value is not dynamic. Set it as Always Dynamic."));
+			}
+		});
+
+	return SequenceEvaluator;
+}
+
 FSequenceEvaluatorReference USequenceEvaluatorLibrary::AdvanceTime(const FAnimUpdateContext& UpdateContext, const FSequenceEvaluatorReference& SequenceEvaluator, float PlayRate)
 {
 	SequenceEvaluator.CallAnimNodeFunction<FAnimNode_SequenceEvaluator>(
diff --git a/Engine/Source/Runtime/AnimGraphRuntime/Public/AnimNodes/AnimNode_SequenceEvaluator.h b/Engine/Source/Runtime/AnimGraphRuntime/Public/AnimNodes/AnimNode_SequenceEvaluator.h
--- a/Engine/Source/Runtime/AnimGraphRuntime/Public/AnimNodes/AnimNode_SequenceEvaluator.h
+++ b/Engine/Source/Runtime/AnimGraphRuntime/Public/AnimNodes/AnimNode_SequenceEvaluator.h
@@ -126,9 +126,17 @@ private:
 	TObjectPtr<UAnimSequenceBase> Sequence = nullptr;
 
 	// The time at which to evaluate the associated sequence
-	UPROPERTY(EditAnywhere, Category=Settings, meta=(PinShownByDefault, FoldProperty))
+	UPROPERTY(EditAnywhere, Category=Settings, meta=(PinShownByDefault, FoldProperty, EditCondition="!bUseExplicitFrame"))
 	float ExplicitTime = 0.0f;
 
+	// Whether to use ExplicitFrame (true) or ExplicitTime (false) when evaluating the associated sequence
+	UPROPERTY(EditAnywhere, Category=Settings, meta=(NeverAsPin, FoldProperty))
+	bool bUseExplicitFrame = false;
+
+	// The frame at which to evaluate the associated sequence
+	UPROPERTY(EditAnywhere, Category=Settings, meta=(PinHiddenByDefault, FoldProperty, EditCondition="bUseExplicitFrame"))
+	int32 ExplicitFrame = 0;
+
 	/** This only works if bTeleportToExplicitTime is false OR this node is set to use SyncGroup */
 	UPROPERTY(EditAnywhere, Category = Settings, meta = (PinHiddenByDefault, FoldProperty))
 	bool bShouldLoop = true;
@@ -169,6 +177,11 @@ public:
 	ANIMGRAPHRUNTIME_API virtual bool SetGroupRole(EAnimGroupRole::Type InRole) override;
 	ANIMGRAPHRUNTIME_API virtual bool SetGroupMethod(EAnimSyncMethod InMethod) override;
 	ANIMGRAPHRUNTIME_API virtual bool SetIgnoreForRelevancyTest(bool bInIgnoreForRelevancyTest) override;
+	
+	ANIMGRAPHRUNTIME_API bool ShouldUseExplicitFrame() const;
+	ANIMGRAPHRUNTIME_API bool SetShouldUseExplicitFrame(bool bFlag);
+	ANIMGRAPHRUNTIME_API bool SetExplicitFrame(int32 InFrame);
+	ANIMGRAPHRUNTIME_API int32 GetExplicitFrame() const;
 };
 
 // Sequence evaluator node that can be used standalone (without constant folding)
@@ -199,9 +212,15 @@ private:
 	TObjectPtr<UAnimSequenceBase> Sequence = nullptr;
 
 	// The time at which to evaluate the associated sequence
-	UPROPERTY(EditAnywhere, Category=Settings, meta=(PinShownByDefault))
+	UPROPERTY(EditAnywhere, Category=Settings, meta=(PinShownByDefault, EditCondition="!bUseFrame"))
 	float ExplicitTime = 0.0f;
 
+	UPROPERTY(EditAnywhere, Category=Settings, meta=(NeverAsPin))
+	bool bUseExplicitFrame = false;
+	
+	UPROPERTY(EditAnywhere, Category=Settings, meta=(PinHiddenByDefault, EditCondition="bUseFrame"))
+	int32 ExplicitFrame = 0;
+
 	/** This only works if bTeleportToExplicitTime is false OR this node is set to use SyncGroup */
 	UPROPERTY(EditAnywhere, Category = Settings, meta = (PinHiddenByDefault))
 	bool bShouldLoop = true;
@@ -227,7 +246,19 @@ public:
 	virtual void SetTeleportToExplicitTime(bool bInTeleport) override { bTeleportToExplicitTime = bInTeleport; }
 	virtual void SetReinitializationBehavior(TEnumAsByte<ESequenceEvalReinit::Type> InBehavior) override { ReinitializationBehavior = InBehavior; }
 	virtual UAnimSequenceBase* GetSequence() const override { return Sequence; }
-	virtual float GetExplicitTime() const override { return ExplicitTime; }
+	virtual float GetExplicitTime() const override
+	{
+		if(ShouldUseExplicitFrame())
+		{
+			if(const UAnimSequenceBase* SequenceBase = GetSequence())
+			{
+				return SequenceBase->GetSamplingFrameRate().AsSeconds(ExplicitFrame);
+			}
+
+			return 0.f;
+		}
+		return ExplicitTime;
+	}
 	virtual bool GetTeleportToExplicitTime() const override { return bTeleportToExplicitTime; }
 	virtual TEnumAsByte<ESequenceEvalReinit::Type> GetReinitializationBehavior() const override { return ReinitializationBehavior; }
 	virtual float GetStartPosition() const override { return StartPosition; }
@@ -242,4 +273,9 @@ public:
 	virtual bool SetGroupRole(EAnimGroupRole::Type InRole) override { GroupRole = InRole; return true; }
 	virtual bool SetGroupMethod(EAnimSyncMethod InMethod) override { Method = InMethod; return true; }
 	virtual bool SetIgnoreForRelevancyTest(bool bInIgnoreForRelevancyTest) override { bIgnoreForRelevancyTest = bInIgnoreForRelevancyTest; return true; }
+
+	bool ShouldUseExplicitFrame() const { return bUseExplicitFrame; }
+	void SetShouldUseExplicitFrame(bool bFlag) { bUseExplicitFrame = bFlag; }
+	void SetExplicitFrame(int32 InFrame) { ExplicitFrame = InFrame; }
+	int32 GetExplicitFrame() const { return ExplicitFrame; }
 };
diff --git a/Engine/Source/Runtime/AnimGraphRuntime/Public/SequenceEvaluatorLibrary.h b/Engine/Source/Runtime/AnimGraphRuntime/Public/SequenceEvaluatorLibrary.h
--- a/Engine/Source/Runtime/AnimGraphRuntime/Public/SequenceEvaluatorLibrary.h
+++ b/Engine/Source/Runtime/AnimGraphRuntime/Public/SequenceEvaluatorLibrary.h
@@ -44,6 +44,10 @@ public:
 	UFUNCTION(BlueprintCallable, Category = "Animation|Sequences", meta=(BlueprintThreadSafe))
 	static ANIMGRAPHRUNTIME_API FSequenceEvaluatorReference SetExplicitTime(const FSequenceEvaluatorReference& SequenceEvaluator, float Time);
 
+	/** Set the current accumulated time, using a frame number, of the sequence evaluator */
+	UFUNCTION(BlueprintCallable, Category = "Animation|Sequences", meta=(BlueprintThreadSafe))
+	static ANIMGRAPHRUNTIME_API FSequenceEvaluatorReference SetExplicitFrame(const FSequenceEvaluatorReference& SequenceEvaluator, int32 Frame);
+
 	/** Advance the current accumulated time of the sequence evaluator */
 	UFUNCTION(BlueprintCallable, Category = "Animation|Sequences", meta = (BlueprintThreadSafe))
 	static ANIMGRAPHRUNTIME_API FSequenceEvaluatorReference AdvanceTime(const FAnimUpdateContext& UpdateContext, const FSequenceEvaluatorReference& SequenceEvaluator, float PlayRate = 1.0f);
