[UE] RBAN Add CVar to force the use of the default deferred simulation CVar to determine whether to defer simulation, ignoring any settings in individual nodes.
- This can be used to force deferred simulation on or off for all RBAN

diff --git a/Engine/Source/Runtime/AnimGraphRuntime/Private/BoneControllers/AnimNode_RigidBody.cpp b/Engine/Source/Runtime/AnimGraphRuntime/Private/BoneControllers/AnimNode_RigidBody.cpp
--- a/Engine/Source/Runtime/AnimGraphRuntime/Private/BoneControllers/AnimNode_RigidBody.cpp
+++ b/Engine/Source/Runtime/AnimGraphRuntime/Private/BoneControllers/AnimNode_RigidBody.cpp
@@ -90,6 +90,9 @@ FAutoConsoleVariableRef CVarRigidBodyNodeDeferredSimulationDefault(
 	TEXT("Whether rigid body simulations are deferred one frame for assets that don't opt into a specific simulation timing"),
 	ECVF_Default);
 
+bool bRBAN_DeferredSimulationForceDefault = false;
+FAutoConsoleVariableRef CVarRigidBodyNodeDeferredSimulationForceDefault(TEXT("p.RigidBodyNode.DeferredSimulationForceDefault"), bRBAN_DeferredSimulationForceDefault, TEXT("When true, rigid body simulation will always use the value of p.RigidBodyNode.DeferredSimulationDefault to determine whether to defer the simulation work, ignoring the setting in the individual node."), ECVF_Default);
+
 bool bRBAN_DebugDraw = false;
 FAutoConsoleVariableRef CVarRigidBodyNodeDebugDraw(TEXT("p.RigidBodyNode.DebugDraw"), bRBAN_DebugDraw, TEXT("Whether to debug draw the rigid body simulation state. Requires p.Chaos.DebugDraw.Enabled 1 to function as well."), ECVF_Default);
 
@@ -689,8 +692,8 @@ void FAnimNode_RigidBody::EvaluateSkeletalControl_AnyThread(FComponentSpacePoseC
 
 		// Assets can override config for deferred simulation
 		const bool bUseDeferredSimulationTask =
-			(SimulationTiming == ESimulationTiming::Deferred) ||
-			((SimulationTiming == ESimulationTiming::Default) && bRBAN_DeferredSimulationDefault);
+			((SimulationTiming == ESimulationTiming::Default) || bRBAN_DeferredSimulationForceDefault) ? bRBAN_DeferredSimulationDefault : (SimulationTiming == ESimulationTiming::Deferred);
+
 		FVector SimSpaceGravity(0.f);
 
 		// Only need to tick physics if we didn't reset and we have some time to simulate
