Optimize common case of single full-weight entry in a blend list

Prior to this change, a single entry would needlessly multiply a pose with weight 1

diff --git a/Engine/Source/Runtime/AnimGraphRuntime/Private/AnimNodes/AnimNode_BlendListBase.cpp b/Engine/Source/Runtime/AnimGraphRuntime/Private/AnimNodes/AnimNode_BlendListBase.cpp
--- a/Engine/Source/Runtime/AnimGraphRuntime/Private/AnimNodes/AnimNode_BlendListBase.cpp
+++ b/Engine/Source/Runtime/AnimGraphRuntime/Private/AnimNodes/AnimNode_BlendListBase.cpp
@@ -289,53 +289,61 @@ void FAnimNode_BlendListBase::Evaluate_AnyThread(FPoseContext& Output)
 	const int32 NumPoses = PosesToEvaluate.Num();
 
 	if ((NumPoses > 0) && (BlendPose.Num() == BlendWeights.Num()))
-	{		
-		// Scratch arrays for evaluation, stack allocated
-		TArray<FCompactPose, TInlineAllocator<8>> FilteredPoses;
-		TArray<FBlendedCurve, TInlineAllocator<8>> FilteredCurve;
-		TArray<UE::Anim::FStackAttributeContainer, TInlineAllocator<8>> FilteredAttributes;
-
-		FilteredPoses.SetNum(NumPoses, false);
-		FilteredCurve.SetNum(NumPoses, false);
-		FilteredAttributes.SetNum(NumPoses, false);
-
-		int32 NumActivePoses = 0;
-		for (int32 i = 0; i < PosesToEvaluate.Num(); ++i)
+	{
+		UBlendProfile* CurrentBlendProfile = GetBlendProfile();
+		if(NumPoses == 1 && FAnimWeight::IsFullWeight(BlendWeights[0]) && CurrentBlendProfile == nullptr)
+		{
+			// Single full weight pose - pass-through fast common case
+			BlendPose[PosesToEvaluate[0]].Evaluate(Output);
+		}
+		else
 		{
-			int32 PoseIndex = PosesToEvaluate[i];
+			// Scratch arrays for evaluation, stack allocated
+			TArray<FCompactPose, TInlineAllocator<8>> FilteredPoses;
+			TArray<FBlendedCurve, TInlineAllocator<8>> FilteredCurve;
+			TArray<UE::Anim::FStackAttributeContainer, TInlineAllocator<8>> FilteredAttributes;
 
-			FPoseContext EvaluateContext(Output);
+			FilteredPoses.SetNum(NumPoses, false);
+			FilteredCurve.SetNum(NumPoses, false);
+			FilteredAttributes.SetNum(NumPoses, false);
 
-			FPoseLink& CurrentPose = BlendPose[PoseIndex];
-			CurrentPose.Evaluate(EvaluateContext);
+			int32 NumActivePoses = 0;
+			for (int32 i = 0; i < PosesToEvaluate.Num(); ++i)
+			{
+				int32 PoseIndex = PosesToEvaluate[i];
 
-			FilteredPoses[i].MoveBonesFrom(EvaluateContext.Pose);
-			FilteredCurve[i].MoveFrom(EvaluateContext.Curve);
-			FilteredAttributes[i].MoveFrom(EvaluateContext.CustomAttributes);
-		}
+				FPoseContext EvaluateContext(Output);
 
-		FAnimationPoseData OutAnimationPoseData(Output);
+				FPoseLink& CurrentPose = BlendPose[PoseIndex];
+				CurrentPose.Evaluate(EvaluateContext);
 
-		// Use the calculated blend sample data if we're blending per-bone
-		UBlendProfile* CurrentBlendProfile = GetBlendProfile();
-		if (CurrentBlendProfile)
-		{
-			const USkeleton* TargetSkeleton = Output.Pose.GetBoneContainer().GetSkeletonAsset();
-			const USkeleton* SourceSkeleton = CurrentBlendProfile->OwningSkeleton;
-			const FSkeletonRemapping& SkeletonRemapping = UE::Anim::FSkeletonRemappingRegistry::Get().GetRemapping(SourceSkeleton, TargetSkeleton);
-			if (SkeletonRemapping.IsValid())
+				FilteredPoses[i].MoveBonesFrom(EvaluateContext.Pose);
+				FilteredCurve[i].MoveFrom(EvaluateContext.Curve);
+				FilteredAttributes[i].MoveFrom(EvaluateContext.CustomAttributes);
+			}
+
+			FAnimationPoseData OutAnimationPoseData(Output);
+		
+			// Use the calculated blend sample data if we're blending per-bone
+			if (CurrentBlendProfile)
 			{
-				FAnimationRuntime::BlendPosesTogetherPerBoneRemapped(FilteredPoses, FilteredCurve, FilteredAttributes, CurrentBlendProfile, PerBoneSampleData, PosesToEvaluate, SkeletonRemapping, OutAnimationPoseData);
+				const USkeleton* TargetSkeleton = Output.Pose.GetBoneContainer().GetSkeletonAsset();
+				const USkeleton* SourceSkeleton = CurrentBlendProfile->OwningSkeleton;
+				const FSkeletonRemapping& SkeletonRemapping = UE::Anim::FSkeletonRemappingRegistry::Get().GetRemapping(SourceSkeleton, TargetSkeleton);
+				if (SkeletonRemapping.IsValid())
+				{
+					FAnimationRuntime::BlendPosesTogetherPerBoneRemapped(FilteredPoses, FilteredCurve, FilteredAttributes, CurrentBlendProfile, PerBoneSampleData, PosesToEvaluate, SkeletonRemapping, OutAnimationPoseData);
+				}
+				else
+				{
+					FAnimationRuntime::BlendPosesTogetherPerBone(FilteredPoses, FilteredCurve, FilteredAttributes, CurrentBlendProfile, PerBoneSampleData, PosesToEvaluate, OutAnimationPoseData);
+				}
 			}
 			else
 			{
-				FAnimationRuntime::BlendPosesTogetherPerBone(FilteredPoses, FilteredCurve, FilteredAttributes, CurrentBlendProfile, PerBoneSampleData, PosesToEvaluate, OutAnimationPoseData);
+				FAnimationRuntime::BlendPosesTogether(FilteredPoses, FilteredCurve, FilteredAttributes, BlendWeights, PosesToEvaluate, OutAnimationPoseData);
 			}
 		}
-		else
-		{
-			FAnimationRuntime::BlendPosesTogether(FilteredPoses, FilteredCurve, FilteredAttributes, BlendWeights, PosesToEvaluate, OutAnimationPoseData);
-		}		
 	}
 	else
 	{
