Fix blend list optimization to correctly check weight of single active pose

diff --git a/Engine/Source/Runtime/AnimGraphRuntime/Private/AnimNodes/AnimNode_BlendListBase.cpp b/Engine/Source/Runtime/AnimGraphRuntime/Private/AnimNodes/AnimNode_BlendListBase.cpp
--- a/Engine/Source/Runtime/AnimGraphRuntime/Private/AnimNodes/AnimNode_BlendListBase.cpp
+++ b/Engine/Source/Runtime/AnimGraphRuntime/Private/AnimNodes/AnimNode_BlendListBase.cpp
@@ -291,7 +291,7 @@ void FAnimNode_BlendListBase::Evaluate_AnyThread(FPoseContext& Output)
 	if ((NumPoses > 0) && (BlendPose.Num() == BlendWeights.Num()))
 	{
 		UBlendProfile* CurrentBlendProfile = GetBlendProfile();
-		if(NumPoses == 1 && FAnimWeight::IsFullWeight(BlendWeights[0]) && CurrentBlendProfile == nullptr)
+		if(NumPoses == 1 && FAnimWeight::IsFullWeight(BlendWeights[PosesToEvaluate[0]]) && CurrentBlendProfile == nullptr)
 		{
 			// Single full weight pose - pass-through fast common case
 			BlendPose[PosesToEvaluate[0]].Evaluate(Output);
