RBAN can now expose the override physics asset as a pin

* Detects changes on the physics override asset in case it gets piped in as a pin and re-initializes physics along with the bone references.
* Added a new property to control if the physics asset shall be defaulted to the one set to the skeletal mesh in case the overwrite physics asset is not specified. Default is true to ensure behavior doesn't change.

diff --git a/Engine/Source/Runtime/AnimGraphRuntime/Private/BoneControllers/AnimNode_RigidBody.cpp b/Engine/Source/Runtime/AnimGraphRuntime/Private/BoneControllers/AnimNode_RigidBody.cpp
--- a/Engine/Source/Runtime/AnimGraphRuntime/Private/BoneControllers/AnimNode_RigidBody.cpp
+++ b/Engine/Source/Runtime/AnimGraphRuntime/Private/BoneControllers/AnimNode_RigidBody.cpp
@@ -974,6 +974,23 @@ void ComputeBodyInsertionOrder(TArray<FBoneIndexType>& InsertionOrder, const USk
 	}
 }
 
+UPhysicsAsset* FAnimNode_RigidBody::GetPhysicsAssetToBeUsed(const UAnimInstance* InAnimInstance) const
+{
+	if (OverridePhysicsAsset)
+	{
+		return ToRawPtr(OverridePhysicsAsset);
+	}
+
+	ensure(InAnimInstance);
+	const USkeletalMeshComponent* SkeletalMeshComp = InAnimInstance->GetSkelMeshComponent();
+	if (bDefaultToSkeletalMeshPhysicsAsset && SkeletalMeshComp)
+	{
+		return SkeletalMeshComp->GetPhysicsAsset();
+	}
+
+	return nullptr;
+}
+
 void FAnimNode_RigidBody::InitPhysics(const UAnimInstance* InAnimInstance)
 {
 	LLM_SCOPE_BYNAME(TEXT("Animation/RigidBody")); 
@@ -994,7 +1011,7 @@ void FAnimNode_RigidBody::InitPhysics(const UAnimInstance* InAnimInstance)
 	}
 
 	const FReferenceSkeleton& SkelMeshRefSkel = SkeletalMeshAsset->GetRefSkeleton();
-	UsePhysicsAsset = OverridePhysicsAsset ? ToRawPtr(OverridePhysicsAsset) : InAnimInstance->GetSkelMeshComponent()->GetPhysicsAsset();
+	UsePhysicsAsset = GetPhysicsAssetToBeUsed(InAnimInstance);
 
 	ensure(SkeletonAsset == SkeletalMeshAsset->GetSkeleton());
 
@@ -1431,6 +1448,16 @@ DECLARE_CYCLE_STAT(TEXT("RigidBody_PreUpdate"), STAT_RigidBody_PreUpdate, STATGR
 
 void FAnimNode_RigidBody::PreUpdate(const UAnimInstance* InAnimInstance)
 {
+	// Detect changes in the physics asset to be used. This can happen when using the override physics asset as a pin on the anim graph node.
+	UPhysicsAsset* PhysicsAssetToBeUsed = GetPhysicsAssetToBeUsed(InAnimInstance);
+	if (UsePhysicsAsset != PhysicsAssetToBeUsed)
+	{
+		InitPhysics(InAnimInstance);
+
+		// Update the bone references after a change in the physics asset. This needs to happen after initializing physics as the Bodies set up in InitPhysics() need to be up to date.
+		InitializeBoneReferences(InAnimInstance->GetRequiredBones());
+	}
+
 	// Don't update geometry if RBN is disabled
 	if(!bEnabled)
 	{
diff --git a/Engine/Source/Runtime/AnimGraphRuntime/Public/BoneControllers/AnimNode_RigidBody.h b/Engine/Source/Runtime/AnimGraphRuntime/Public/BoneControllers/AnimNode_RigidBody.h
--- a/Engine/Source/Runtime/AnimGraphRuntime/Public/BoneControllers/AnimNode_RigidBody.h
+++ b/Engine/Source/Runtime/AnimGraphRuntime/Public/BoneControllers/AnimNode_RigidBody.h
@@ -181,11 +181,18 @@ struct FAnimNode_RigidBody : public FAnimNode_SkeletalControlBase
 	UPhysicsAsset* GetPhysicsAsset() const { return UsePhysicsAsset; }
 
 public:
-	/** Physics asset to use. If empty use the skeletal mesh's default physics asset */
-	UPROPERTY(EditAnywhere, Category = Settings)
+	/** Physics asset to use. If empty use the skeletal mesh's default physics asset in case Default To Skeletal Mesh Physics Asset is set to True. */
+	UPROPERTY(EditAnywhere, Category = Settings, meta = (PinHiddenByDefault))
 	TObjectPtr<UPhysicsAsset> OverridePhysicsAsset;
 
+	/** Use the skeletal mesh physics asset as default in case set to True. The Override Physics Asset will always have priority over this. */
+	UPROPERTY(EditAnywhere, Category = Settings)
+	bool bDefaultToSkeletalMeshPhysicsAsset = true;
+
 private:
+	/** Get the physics asset candidate to be used while respecting the bDefaultToSkeletalMeshPhysicsAsset and the priority to the override physics asset. */
+	UPhysicsAsset* GetPhysicsAssetToBeUsed(const UAnimInstance* InAnimInstance) const;
+
 	FTransform PreviousCompWorldSpaceTM;
 	FTransform CurrentTransform;
 	FTransform PreviousTransform;
