Allow updating both children of the Two Way Blend even when not relevant

diff --git a/Engine/Source/Runtime/AnimGraphRuntime/Private/AnimNodes/AnimNode_TwoWayBlend.cpp b/Engine/Source/Runtime/AnimGraphRuntime/Private/AnimNodes/AnimNode_TwoWayBlend.cpp
--- a/Engine/Source/Runtime/AnimGraphRuntime/Private/AnimNodes/AnimNode_TwoWayBlend.cpp
+++ b/Engine/Source/Runtime/AnimGraphRuntime/Private/AnimNodes/AnimNode_TwoWayBlend.cpp
@@ -94,6 +94,11 @@ void FAnimNode_TwoWayBlend::Update_AnyThread(const FAnimationUpdateContext& Cont
 		}
 		else
 		{
+			if (bAlwaysUpdateChildren)
+			{
+				A.Update(Context.FractionalWeight(FAnimWeight::GetSmallestRelevantWeight()));
+			}
+
 			// Take all of B
 			B.Update(Context);
 		}
@@ -102,6 +107,11 @@ void FAnimNode_TwoWayBlend::Update_AnyThread(const FAnimationUpdateContext& Cont
 	{
 		// Take all of A
 		A.Update(Context);
+
+		if (bAlwaysUpdateChildren)
+		{
+			B.Update(Context.FractionalWeight(FAnimWeight::GetSmallestRelevantWeight()));
+		}
 	}
 
 	TRACE_ANIM_NODE_VALUE(Context, TEXT("Alpha"), InternalBlendAlpha);
diff --git a/Engine/Source/Runtime/AnimGraphRuntime/Public/AnimNodes/AnimNode_TwoWayBlend.h b/Engine/Source/Runtime/AnimGraphRuntime/Public/AnimNodes/AnimNode_TwoWayBlend.h
--- a/Engine/Source/Runtime/AnimGraphRuntime/Public/AnimNodes/AnimNode_TwoWayBlend.h
+++ b/Engine/Source/Runtime/AnimGraphRuntime/Public/AnimNodes/AnimNode_TwoWayBlend.h
@@ -40,6 +40,10 @@ protected:
 	UPROPERTY(EditAnywhere, Category = Option)
 	uint8 bResetChildOnActivation:1;
 
+	/** Always update children, regardless of whether or not that child has weight. */
+	UPROPERTY(EditAnywhere, Category = Option)
+	uint8 bAlwaysUpdateChildren:1;
+
 public:
 	/** The float value that controls the alpha blending when the alpha input type is set to 'Float' */
 	UPROPERTY(EditAnywhere, BlueprintReadWrite, Category=Settings, meta=(PinShownByDefault))
@@ -68,6 +72,7 @@ public:
 		, bAIsRelevant(false)
 		, bBIsRelevant(false)
 		, bResetChildOnActivation(false)
+		, bAlwaysUpdateChildren(false)
 		, Alpha(0.0f)
 		, AlphaCurveName(NAME_None)
 		, InternalBlendAlpha(0.0f)
