RBAN Static analysis fix

diff --git a/Engine/Source/Runtime/AnimGraphRuntime/Private/BoneControllers/AnimNode_RigidBody.cpp b/Engine/Source/Runtime/AnimGraphRuntime/Private/BoneControllers/AnimNode_RigidBody.cpp
--- a/Engine/Source/Runtime/AnimGraphRuntime/Private/BoneControllers/AnimNode_RigidBody.cpp
+++ b/Engine/Source/Runtime/AnimGraphRuntime/Private/BoneControllers/AnimNode_RigidBody.cpp
@@ -981,11 +981,13 @@ UPhysicsAsset* FAnimNode_RigidBody::GetPhysicsAssetToBeUsed(const UAnimInstance*
 		return ToRawPtr(OverridePhysicsAsset);
 	}
 
-	ensure(InAnimInstance);
-	const USkeletalMeshComponent* SkeletalMeshComp = InAnimInstance->GetSkelMeshComponent();
-	if (bDefaultToSkeletalMeshPhysicsAsset && SkeletalMeshComp)
+	if (InAnimInstance)
 	{
-		return SkeletalMeshComp->GetPhysicsAsset();
+		const USkeletalMeshComponent* SkeletalMeshComp = InAnimInstance->GetSkelMeshComponent();
+		if (bDefaultToSkeletalMeshPhysicsAsset && SkeletalMeshComp)
+		{
+			return SkeletalMeshComp->GetPhysicsAsset();
+		}
 	}
 
 	return nullptr;
