Explicitly disable inlining with clang of specific functions to save stack space in hot recursive function

diff --git a/Engine/Source/Runtime/AnimGraphRuntime/Private/BoneControllers/AnimNode_SkeletalControlBase.cpp b/Engine/Source/Runtime/AnimGraphRuntime/Private/BoneControllers/AnimNode_SkeletalControlBase.cpp
--- a/Engine/Source/Runtime/AnimGraphRuntime/Private/BoneControllers/AnimNode_SkeletalControlBase.cpp
+++ b/Engine/Source/Runtime/AnimGraphRuntime/Private/BoneControllers/AnimNode_SkeletalControlBase.cpp
@@ -17,6 +17,14 @@
 #endif
 DEFINE_LOG_CATEGORY(LogSkeletalControlBase);
 
+// Macro to disable inlining of function calls within the statement that follows
+// e.g: UE_DONT_INLINE_CALL ThisFunctionCallWontBeInlined(Foo->ThisIsntInlinedEither());
+#if defined(__clang__)
+#define UE_DONT_INLINE_CALL [[clang::noinline]]
+#else
+#define UE_DONT_INLINE_CALL
+#endif
+
 /////////////////////////////////////////////////////
 // FAnimNode_SkeletalControlBase
 
@@ -56,8 +64,21 @@ void FAnimNode_SkeletalControlBase::UpdateComponentPose_AnyThread(const FAnimati
 
 void FAnimNode_SkeletalControlBase::Update_AnyThread(const FAnimationUpdateContext& Context)
 {
+	//////////////////////////////////////////////////////////////////////////
+	// PERFORMANCE CRITICAL NOTE
+	// 
+	// This function is called recursively as we traverse nodes, as such, it is critical to keep the
+	// amount of stack space used to a minimum as many nodes can be traversed. Using too much stack
+	// here can quickly lead to stack overflows.
+	// 
+	// We explicitly disable inlineing for virtual calls. Normally, virtual calls are never inlined
+	// but when PGO and LTO are enabled, the compiler can speculatively de-virtualize the calls. It does
+	// this by comparing the v-table pointer and inlineing the call directly in here with a static jump.
+	// As a result, it can significantly increase the amount of stack space used.
+	//////////////////////////////////////////////////////////////////////////
+
 	DECLARE_SCOPE_HIERARCHICAL_COUNTER_ANIMNODE(Update_AnyThread)
-	UpdateComponentPose_AnyThread(Context);
+	UE_DONT_INLINE_CALL UpdateComponentPose_AnyThread(Context);
 
 	ActualAlpha = 0.f;
 	if (IsLODEnabled(Context.AnimInstanceProxy))
@@ -84,9 +105,17 @@ void FAnimNode_SkeletalControlBase::Update_AnyThread(const FAnimationUpdateConte
 		// Make sure Alpha is clamped between 0 and 1.
 		ActualAlpha = FMath::Clamp<float>(ActualAlpha, 0.f, 1.f);
 
-		if (FAnimWeight::IsRelevant(ActualAlpha) && IsValidToEvaluate(Context.AnimInstanceProxy->GetSkeleton(), Context.AnimInstanceProxy->GetRequiredBones()))
+		if (FAnimWeight::IsRelevant(ActualAlpha))
 		{
-			UpdateInternal(Context);
+			const USkeleton* Skeleton = Context.AnimInstanceProxy->GetSkeleton();
+			const FBoneContainer& RequiredBones = Context.AnimInstanceProxy->GetRequiredBones();
+
+			bool bIsValidToEvaluate;
+			UE_DONT_INLINE_CALL bIsValidToEvaluate = IsValidToEvaluate(Skeleton, RequiredBones);
+			if (bIsValidToEvaluate)
+			{
+				UE_DONT_INLINE_CALL UpdateInternal(Context);
+			}
 		}
 	}
 
@@ -119,13 +148,26 @@ void FAnimNode_SkeletalControlBase::EvaluateComponentSpaceInternal(FComponentSpa
 
 void FAnimNode_SkeletalControlBase::EvaluateComponentSpace_AnyThread(FComponentSpacePoseContext& Output)
 {
+	//////////////////////////////////////////////////////////////////////////
+	// PERFORMANCE CRITICAL NOTE
+	// 
+	// This function is called recursively as we traverse nodes, as such, it is critical to keep the
+	// amount of stack space used to a minimum as many nodes can be traversed. Using too much stack
+	// here can quickly lead to stack overflows.
+	// 
+	// We explicitly disable inlineing for virtual calls. Normally, virtual calls are never inlined
+	// but when PGO and LTO are enabled, the compiler can speculatively de-virtualize the calls. It does
+	// this by comparing the v-table pointer and inlineing the call directly in here with a static jump.
+	// As a result, it can significantly increase the amount of stack space used.
+	//////////////////////////////////////////////////////////////////////////
+
 	DECLARE_SCOPE_HIERARCHICAL_COUNTER_ANIMNODE(EvaluateComponentSpace_AnyThread)
 	ANIM_MT_SCOPE_CYCLE_COUNTER_VERBOSE(SkeletalControlBase, !IsInGameThread());
 
 	// Cache the incoming node IDs in a base context
 	FAnimationBaseContext CachedContext(Output);
 
-	EvaluateComponentPose_AnyThread(Output);
+	UE_DONT_INLINE_CALL EvaluateComponentPose_AnyThread(Output);
 
 #if WITH_EDITORONLY_DATA
 	// save current pose before applying skeletal control to compute the exact gizmo location in AnimGraphNode
@@ -138,22 +180,30 @@ void FAnimNode_SkeletalControlBase::EvaluateComponentSpace_AnyThread(FComponentS
 #endif
 
 	// Apply the skeletal control if it's valid
-	if (FAnimWeight::IsRelevant(ActualAlpha) && IsValidToEvaluate(Output.AnimInstanceProxy->GetSkeleton(), Output.AnimInstanceProxy->GetRequiredBones()))
+	if (FAnimWeight::IsRelevant(ActualAlpha))
 	{
-		Output.SetNodeIds(CachedContext);
+		const USkeleton* Skeleton = Output.AnimInstanceProxy->GetSkeleton();
+		const FBoneContainer& RequiredBones = Output.AnimInstanceProxy->GetRequiredBones();
 
-		EvaluateComponentSpaceInternal(Output);
+		bool bIsValidToEvaluate;
+		UE_DONT_INLINE_CALL bIsValidToEvaluate = IsValidToEvaluate(Skeleton, RequiredBones);
+		if (bIsValidToEvaluate)
+		{
+			Output.SetNodeIds(CachedContext);
 
-		BoneTransforms.Reset(BoneTransforms.Num());
-		EvaluateSkeletalControl_AnyThread(Output, BoneTransforms);
+			UE_DONT_INLINE_CALL EvaluateComponentSpaceInternal(Output);
 
-		if (BoneTransforms.Num() > 0)
-		{
-			const float BlendWeight = FMath::Clamp<float>(ActualAlpha, 0.f, 1.f);
-			Output.Pose.LocalBlendCSBoneTransforms(BoneTransforms, BlendWeight);
-		}
+			BoneTransforms.Reset(BoneTransforms.Num());
+			UE_DONT_INLINE_CALL EvaluateSkeletalControl_AnyThread(Output, BoneTransforms);
 
-		// we check NaN when you get out of this function in void FComponentSpacePoseLink::EvaluateComponentSpace(FComponentSpacePoseContext& Output)
+			if (BoneTransforms.Num() > 0)
+			{
+				const float BlendWeight = FMath::Clamp<float>(ActualAlpha, 0.f, 1.f);
+				Output.Pose.LocalBlendCSBoneTransforms(BoneTransforms, BlendWeight);
+			}
+
+			// we check NaN when you get out of this function in void FComponentSpacePoseLink::EvaluateComponentSpace(FComponentSpacePoseContext& Output)
+		}
 	}
 }
 
