Crash fix in the layered bone blend node, basically undoing previous optimizations.

#trivial
#swarm none

diff --git a/Engine/Source/Runtime/AnimGraphRuntime/Private/AnimNodes/AnimNode_LayeredBoneBlend.cpp b/Engine/Source/Runtime/AnimGraphRuntime/Private/AnimNodes/AnimNode_LayeredBoneBlend.cpp
--- a/Engine/Source/Runtime/AnimGraphRuntime/Private/AnimNodes/AnimNode_LayeredBoneBlend.cpp
+++ b/Engine/Source/Runtime/AnimGraphRuntime/Private/AnimNodes/AnimNode_LayeredBoneBlend.cpp
@@ -182,13 +182,6 @@ void FAnimNode_LayeredBoneBlend::Update_AnyThread(const FAnimationUpdateContext&
 	TRACE_ANIM_NODE_VALUE(Context, TEXT("Num Poses"), BlendPoses.Num());
 }
 
-struct FLayeredBoneBlendScratchArea : TThreadSingleton<FLayeredBoneBlendScratchArea>
-{
-	TArray<FCompactPose> TargetBlendPoses;
-	TArray<FBlendedCurve> TargetBlendCurves;
-	TArray<UE::Anim::FStackAttributeContainer> TargetBlendAttributes;
-};
-
 void FAnimNode_LayeredBoneBlend::Evaluate_AnyThread(FPoseContext& Output)
 {
 	DECLARE_SCOPE_HIERARCHICAL_COUNTER_ANIMNODE(Evaluate_AnyThread)
@@ -206,18 +199,13 @@ void FAnimNode_LayeredBoneBlend::Evaluate_AnyThread(FPoseContext& Output)
 		// evaluate children
 		BasePose.Evaluate(BasePoseContext);
 
-		FLayeredBoneBlendScratchArea& ScratchArea = FLayeredBoneBlendScratchArea::Get();
-		
-		TArray<FCompactPose>& TargetBlendPoses = ScratchArea.TargetBlendPoses;
-		TargetBlendPoses.Reset();
+		TArray<FCompactPose> TargetBlendPoses;
 		TargetBlendPoses.SetNum(NumPoses);
 
-		TArray<FBlendedCurve>& TargetBlendCurves = ScratchArea.TargetBlendCurves;
-		TargetBlendCurves.Reset();
+		TArray<FBlendedCurve> TargetBlendCurves;
 		TargetBlendCurves.SetNum(NumPoses);
 
-		TArray<UE::Anim::FStackAttributeContainer>& TargetBlendAttributes = ScratchArea.TargetBlendAttributes;
-		TargetBlendAttributes.Reset();
+		TArray<UE::Anim::FStackAttributeContainer> TargetBlendAttributes;
 		TargetBlendAttributes.SetNum(NumPoses);
 
 		for (int32 ChildIndex = 0; ChildIndex < NumPoses; ++ChildIndex)
