Runtime switching override physics assets on rigid body animation node using function library

A few days ago I've added runtime switching for the override physics assets by exposing the override physics assets as a pin. To align better with the pattern we use with larger objects, I've now changed that to a function library.

diff --git a/Engine/Source/Runtime/AnimGraphRuntime/Private/BoneControllers/AnimNode_RigidBody.cpp b/Engine/Source/Runtime/AnimGraphRuntime/Private/BoneControllers/AnimNode_RigidBody.cpp
--- a/Engine/Source/Runtime/AnimGraphRuntime/Private/BoneControllers/AnimNode_RigidBody.cpp
+++ b/Engine/Source/Runtime/AnimGraphRuntime/Private/BoneControllers/AnimNode_RigidBody.cpp
@@ -1460,6 +1460,11 @@ void FAnimNode_RigidBody::ResetDynamics(ETeleportType InTeleportType)
 	ResetSimulatedTeleportType = ((InTeleportType > ResetSimulatedTeleportType) ? InTeleportType : ResetSimulatedTeleportType);
 }
 
+void FAnimNode_RigidBody::SetOverridePhysicsAsset(UPhysicsAsset* PhysicsAsset)
+{
+	OverridePhysicsAsset = PhysicsAsset;
+}
+
 DECLARE_CYCLE_STAT(TEXT("RigidBody_PreUpdate"), STAT_RigidBody_PreUpdate, STATGROUP_Anim);
 
 void FAnimNode_RigidBody::PreUpdate(const UAnimInstance* InAnimInstance)
diff --git a/Engine/Source/Runtime/AnimGraphRuntime/Private/BoneControllers/AnimNode_RigidBody_Library.cpp b/Engine/Source/Runtime/AnimGraphRuntime/Private/BoneControllers/AnimNode_RigidBody_Library.cpp
new file mode 100644
--- /dev/null
+++ b/Engine/Source/Runtime/AnimGraphRuntime/Private/BoneControllers/AnimNode_RigidBody_Library.cpp
@@ -0,0 +1,27 @@
+// Copyright Epic Games, Inc. All Rights Reserved.
+
+#include "BoneControllers/AnimNode_RigidBody_Library.h"
+
+FRigidBodyAnimNodeReference UAnimNodeRigidBodyLibrary::ConvertToRigidBodyAnimNode(const FAnimNodeReference& Node, EAnimNodeReferenceConversionResult& Result)
+{
+	return FAnimNodeReference::ConvertToType<FRigidBodyAnimNodeReference>(Node, Result);
+}
+
+void UAnimNodeRigidBodyLibrary::ConvertToRigidBodyAnimNodePure(const FAnimNodeReference& Node, FRigidBodyAnimNodeReference& RigidBodyAnimNode, bool& Result)
+{
+	EAnimNodeReferenceConversionResult ConversionResult;
+	RigidBodyAnimNode = ConvertToRigidBodyAnimNode(Node, ConversionResult);
+	Result = (ConversionResult == EAnimNodeReferenceConversionResult::Succeeded);
+}
+
+FRigidBodyAnimNodeReference UAnimNodeRigidBodyLibrary::SetOverridePhysicsAsset(const FRigidBodyAnimNodeReference& Node, UPhysicsAsset* PhysicsAsset)
+{
+	Node.CallAnimNodeFunction<FAnimNode_RigidBody>(
+		TEXT("SetOverridePhysicsAsset"),
+		[PhysicsAsset](FAnimNode_RigidBody& InRigidBodyNode)
+		{
+			InRigidBodyNode.SetOverridePhysicsAsset(PhysicsAsset);
+		});
+
+	return Node;
+}
\ No newline at end of file
diff --git a/Engine/Source/Runtime/AnimGraphRuntime/Public/BoneControllers/AnimNode_RigidBody.h b/Engine/Source/Runtime/AnimGraphRuntime/Public/BoneControllers/AnimNode_RigidBody.h
--- a/Engine/Source/Runtime/AnimGraphRuntime/Public/BoneControllers/AnimNode_RigidBody.h
+++ b/Engine/Source/Runtime/AnimGraphRuntime/Public/BoneControllers/AnimNode_RigidBody.h
@@ -178,11 +178,17 @@ struct FAnimNode_RigidBody : public FAnimNode_SkeletalControlBase
 	// TEMP: Exposed for use in PhAt as a quick way to get drag handles working with Chaos
 	virtual ImmediatePhysics::FSimulation* GetSimulation() { return PhysicsSimulation; }
 
+	/**
+	 * Set the override physics asset. This will automatically trigger a physics re-init in case the override physics asset changes. 
+	 * Users can get access to this in the Animation Blueprint via the Animation Node Functions.
+	 */
+	void SetOverridePhysicsAsset(UPhysicsAsset* PhysicsAsset);
+
 	UPhysicsAsset* GetPhysicsAsset() const { return UsePhysicsAsset; }
 
 public:
 	/** Physics asset to use. If empty use the skeletal mesh's default physics asset in case Default To Skeletal Mesh Physics Asset is set to True. */
-	UPROPERTY(EditAnywhere, Category = Settings, meta = (PinHiddenByDefault))
+	UPROPERTY(EditAnywhere, Category = Settings)
 	TObjectPtr<UPhysicsAsset> OverridePhysicsAsset;
 
 	/** Use the skeletal mesh physics asset as default in case set to True. The Override Physics Asset will always have priority over this. */
diff --git a/Engine/Source/Runtime/AnimGraphRuntime/Public/BoneControllers/AnimNode_RigidBody_Library.h b/Engine/Source/Runtime/AnimGraphRuntime/Public/BoneControllers/AnimNode_RigidBody_Library.h
new file mode 100644
--- /dev/null
+++ b/Engine/Source/Runtime/AnimGraphRuntime/Public/BoneControllers/AnimNode_RigidBody_Library.h
@@ -0,0 +1,38 @@
+// Copyright Epic Games, Inc. All Rights Reserved.
+
+#pragma once
+
+#include "CoreMinimal.h"
+
+#include "Kismet/BlueprintFunctionLibrary.h"
+#include "Animation/AnimNodeReference.h"
+#include "BoneControllers/AnimNode_RigidBody.h"
+#include "AnimNode_RigidBody_Library.generated.h"
+
+USTRUCT(BlueprintType)
+struct FRigidBodyAnimNodeReference : public FAnimNodeReference
+{
+	GENERATED_BODY()
+
+	typedef FAnimNode_RigidBody FInternalNodeType;
+};
+
+// Exposes operations to be performed on a rigid body anim node
+UCLASS(Experimental, MinimalAPI)
+class UAnimNodeRigidBodyLibrary : public UBlueprintFunctionLibrary
+{
+	GENERATED_BODY()
+
+public:
+	/** Get a rigid body anim node context from an anim node context */
+	UFUNCTION(BlueprintCallable, Category = "Animation|Dynamics", meta = (BlueprintThreadSafe, ExpandEnumAsExecs = "Result"))
+	static ANIMGRAPHRUNTIME_API FRigidBodyAnimNodeReference ConvertToRigidBodyAnimNode(const FAnimNodeReference& Node, EAnimNodeReferenceConversionResult& Result);
+
+	/** Get a rigid body anim node context from an anim node context (pure) */
+	UFUNCTION(BlueprintPure, Category = "Animation|Dynamics", meta = (BlueprintThreadSafe, DisplayName = "Convert to rigid body"))
+	static void ConvertToRigidBodyAnimNodePure(const FAnimNodeReference& Node, FRigidBodyAnimNodeReference& RigidBodyAnimNode, bool& Result);
+
+	/** Set the physics asset on the rigid body anim graph node (RBAN). */
+	UFUNCTION(BlueprintCallable, Category = "Animation|Dynamics", meta = (BlueprintThreadSafe))
+	static ANIMGRAPHRUNTIME_API FRigidBodyAnimNodeReference SetOverridePhysicsAsset(const FRigidBodyAnimNodeReference& Node, UPhysicsAsset* PhysicsAsset);
+};
