[Backout] - CL29082994 - CIS Error
Original CL Desc
-----------------------------------------------------------------
Runtime switching override physics assets on rigid body animation node using function library

A few days ago I've added runtime switching for the override physics assets by exposing the override physics assets as a pin. To align better with the pattern we use with larger objects, I've now changed that to a function library.

diff --git a/Engine/Source/Runtime/AnimGraphRuntime/Private/BoneControllers/AnimNode_RigidBody.cpp b/Engine/Source/Runtime/AnimGraphRuntime/Private/BoneControllers/AnimNode_RigidBody.cpp
--- a/Engine/Source/Runtime/AnimGraphRuntime/Private/BoneControllers/AnimNode_RigidBody.cpp
+++ b/Engine/Source/Runtime/AnimGraphRuntime/Private/BoneControllers/AnimNode_RigidBody.cpp
@@ -1460,11 +1460,6 @@ void FAnimNode_RigidBody::ResetDynamics(ETeleportType InTeleportType)
 	ResetSimulatedTeleportType = ((InTeleportType > ResetSimulatedTeleportType) ? InTeleportType : ResetSimulatedTeleportType);
 }
 
-void FAnimNode_RigidBody::SetOverridePhysicsAsset(UPhysicsAsset* PhysicsAsset)
-{
-	OverridePhysicsAsset = PhysicsAsset;
-}
-
 DECLARE_CYCLE_STAT(TEXT("RigidBody_PreUpdate"), STAT_RigidBody_PreUpdate, STATGROUP_Anim);
 
 void FAnimNode_RigidBody::PreUpdate(const UAnimInstance* InAnimInstance)
diff --git a/Engine/Source/Runtime/AnimGraphRuntime/Private/BoneControllers/AnimNode_RigidBody_Library.cpp b/Engine/Source/Runtime/AnimGraphRuntime/Private/BoneControllers/AnimNode_RigidBody_Library.cpp
deleted file mode 100644
diff --git a/Engine/Source/Runtime/AnimGraphRuntime/Public/BoneControllers/AnimNode_RigidBody.h b/Engine/Source/Runtime/AnimGraphRuntime/Public/BoneControllers/AnimNode_RigidBody.h
--- a/Engine/Source/Runtime/AnimGraphRuntime/Public/BoneControllers/AnimNode_RigidBody.h
+++ b/Engine/Source/Runtime/AnimGraphRuntime/Public/BoneControllers/AnimNode_RigidBody.h
@@ -178,17 +178,11 @@ struct FAnimNode_RigidBody : public FAnimNode_SkeletalControlBase
 	// TEMP: Exposed for use in PhAt as a quick way to get drag handles working with Chaos
 	virtual ImmediatePhysics::FSimulation* GetSimulation() { return PhysicsSimulation; }
 
-	/**
-	 * Set the override physics asset. This will automatically trigger a physics re-init in case the override physics asset changes. 
-	 * Users can get access to this in the Animation Blueprint via the Animation Node Functions.
-	 */
-	void SetOverridePhysicsAsset(UPhysicsAsset* PhysicsAsset);
-
 	UPhysicsAsset* GetPhysicsAsset() const { return UsePhysicsAsset; }
 
 public:
 	/** Physics asset to use. If empty use the skeletal mesh's default physics asset in case Default To Skeletal Mesh Physics Asset is set to True. */
-	UPROPERTY(EditAnywhere, Category = Settings)
+	UPROPERTY(EditAnywhere, Category = Settings, meta = (PinHiddenByDefault))
 	TObjectPtr<UPhysicsAsset> OverridePhysicsAsset;
 
 	/** Use the skeletal mesh physics asset as default in case set to True. The Override Physics Asset will always have priority over this. */
diff --git a/Engine/Source/Runtime/AnimGraphRuntime/Public/BoneControllers/AnimNode_RigidBody_Library.h b/Engine/Source/Runtime/AnimGraphRuntime/Public/BoneControllers/AnimNode_RigidBody_Library.h
deleted file mode 100644
