[UE] AnimDynamics CVars to override the clamp on the component linear acceleration and set a global gravity scale
- Added p.AnimDynamics.ComponentAppliedLinearAccClampOverride CVar to override the clamp on the component linear acceleration applied to AnimDynamics simulation for all assets.
- Added p.AnimDynamics.GravityScale Cvar to set a gravity scale that is applied to the default gravity and the gravity override vector.
- Prior to this change, when the sim space was set to World, the gravity override vector was only applied on init. Now it will now be applied every update with any sim space setting.

diff --git a/Engine/Source/Runtime/AnimGraphRuntime/Private/BoneControllers/AnimNode_AnimDynamics.cpp b/Engine/Source/Runtime/AnimGraphRuntime/Private/BoneControllers/AnimNode_AnimDynamics.cpp
--- a/Engine/Source/Runtime/AnimGraphRuntime/Private/BoneControllers/AnimNode_AnimDynamics.cpp
+++ b/Engine/Source/Runtime/AnimGraphRuntime/Private/BoneControllers/AnimNode_AnimDynamics.cpp
@@ -24,6 +24,8 @@ TAutoConsoleVariable<int32> CVarEnableDynamics(TEXT("p.AnimDynamics"), 1, TEXT("
 TAutoConsoleVariable<int32> CVarEnableAdaptiveSubstep(TEXT("p.AnimDynamicsAdaptiveSubstep"), 0, TEXT("Enables/disables adaptive substepping. Adaptive substepping will substep the simulation when it is necessary and maintain a debt buffer for time, always trying to utilise as much time as possible."));
 TAutoConsoleVariable<int32> CVarAdaptiveSubstepNumDebtFrames(TEXT("p.AnimDynamicsNumDebtFrames"), 5, TEXT("Number of frames to maintain as time debt when using adaptive substepping, this should be at least 1 or the time debt will never be cleared."));
 TAutoConsoleVariable<int32> CVarEnableWind(TEXT("p.AnimDynamicsWind"), 1, TEXT("Enables/Disables anim dynamics wind forces globally."), ECVF_Scalability);
+TAutoConsoleVariable<float> CVarComponentAppliedLinearAccClampOverride(TEXT("p.AnimDynamics.ComponentAppliedLinearAccClampOverride"), -1.0f, TEXT("Override the per asset setting for all axis (X,Y & Z) of ComponentAppliedLinearAccClamp for all Anim Dynamics Nodes. Negative values are ignored."));
+TAutoConsoleVariable<float> CVarGravityScale(TEXT("p.AnimDynamics.GravityScale"), 1.0f, TEXT("Multiplies the defalut gravity and the gravity override on all Anim Dynamics Nodes."));
 
 // FindChainBones
 // 
@@ -321,6 +323,7 @@ void FAnimNode_AnimDynamics::EvaluateSkeletalControl_AnyThread(FComponentSpacePo
 			}
 
 			FVector ComponentLinearAcc(0.0f);
+			FVector SimSpaceGravityOverride = GravityOverride;
 
 			if (SimulationSpace != AnimPhysSimSpaceType::World)
 			{
@@ -329,12 +332,7 @@ void FAnimNode_AnimDynamics::EvaluateSkeletalControl_AnyThread(FComponentSpacePo
 				// Transform Gravity Override into simulation space
 				if (bUseGravityOverride && !bGravityOverrideInSimSpace)
 				{
-					const FVector GravityOverrideSimSpace = TransformWorldVectorToSimSpace(Output, GravityOverride);
-
-					for (FAnimPhysRigidBody* ChainBody : SimBodies)
-					{
-						ChainBody->GravityOverride = GravityOverrideSimSpace;
-					}
+					SimSpaceGravityOverride = TransformWorldVectorToSimSpace(Output, SimSpaceGravityOverride);
 				}
 
 				// Calc linear velocity
@@ -349,8 +347,30 @@ void FAnimNode_AnimDynamics::EvaluateSkeletalControl_AnyThread(FComponentSpacePo
 				// Apply opposite acceleration to bodies
 				ComponentLinearAcc += TransformWorldVectorToSimSpace(Output, -ComponentLinearAcceleration) * ComponentLinearAccScale;
 
-				// Clamp to desired strength
-				ComponentLinearAcc = ComponentLinearAcc.BoundToBox(-ComponentAppliedLinearAccClamp, ComponentAppliedLinearAccClamp);
+				// Clamp ComponentLinearAcc to desired strength.	
+				FVector LinearAccClamp = ComponentAppliedLinearAccClamp;
+
+				const float LinearAccClampOverride = CVarComponentAppliedLinearAccClampOverride.GetValueOnAnyThread();
+				if (LinearAccClampOverride >= 0.0f) // Ignore values < 0
+				{
+					LinearAccClamp.Set(LinearAccClampOverride, LinearAccClampOverride, LinearAccClampOverride);
+				}
+
+				ComponentLinearAcc = ComponentLinearAcc.BoundToBox(-LinearAccClamp, LinearAccClamp);	
+			}
+
+			// Update gravity.
+			{
+				const float ExternalGravityScale = CVarGravityScale.GetValueOnAnyThread();
+
+				const FVector AppliedGravityOverride = SimSpaceGravityOverride * ExternalGravityScale;
+				const float AppliedGravityScale = GravityScale * ExternalGravityScale;
+
+				for (FAnimPhysRigidBody* ChainBody : SimBodies)
+				{
+					ChainBody->GravityOverride = AppliedGravityOverride;
+					ChainBody->GravityScale = AppliedGravityScale;
+				}
 			}
 
 			if (CVarEnableAdaptiveSubstep.GetValueOnAnyThread() == 1)
@@ -690,7 +710,8 @@ void FAnimNode_AnimDynamics::InitPhysics(FComponentSpacePoseContext& Output)
 
 		// Transform GravityOverride to simulation space if necessary.
 		const FVector GravityOverrideSimSpace = (bUseGravityOverride && !bGravityOverrideInSimSpace) ? TransformWorldVectorToSimSpace(Output, GravityOverride) : GravityOverride;
-		
+		const float ExternalGravityScale = CVarGravityScale.GetValueOnAnyThread();
+
 		check(PhysicsBodyDefinitions.Num() > 0);
 		if (PhysicsBodyDefinitions.Num() > 0)
 		{
@@ -748,9 +769,9 @@ void FAnimNode_AnimDynamics::InitPhysics(FComponentSpacePoseContext& Output)
 				PhysicsBody.AngularDamping = AngularDampingOverride;
 			}
 
-			PhysicsBody.GravityScale = GravityScale;
+			PhysicsBody.GravityScale = GravityScale * ExternalGravityScale;
 			PhysicsBody.bUseGravityOverride = bUseGravityOverride;
-			PhysicsBody.GravityOverride = GravityOverrideSimSpace;
+			PhysicsBody.GravityOverride = GravityOverrideSimSpace * ExternalGravityScale;
 
 			PhysicsBody.bWindEnabled = bWindWasEnabled;
 
