Fix OOB when accessing SkeletonBoneIndexToBodyIndex array with a negative index

#rnx

diff --git a/Engine/Source/Runtime/AnimGraphRuntime/Private/BoneControllers/AnimNode_RigidBody.cpp b/Engine/Source/Runtime/AnimGraphRuntime/Private/BoneControllers/AnimNode_RigidBody.cpp
--- a/Engine/Source/Runtime/AnimGraphRuntime/Private/BoneControllers/AnimNode_RigidBody.cpp
+++ b/Engine/Source/Runtime/AnimGraphRuntime/Private/BoneControllers/AnimNode_RigidBody.cpp
@@ -1810,10 +1810,9 @@ void FAnimNode_RigidBody::InitializeBoneReferences(const FBoneContainer& Require
 	{
 		const FCompactPoseBoneIndex CompactPoseBoneIndex(Index);
 		const int32 SkeletonBoneIndex = RequiredBones.GetSkeletonIndex(CompactPoseBoneIndex);
-		const int32 IndexToBodyNum = SkeletonBoneIndexToBodyIndex.Num();
 
 		// If we have a missing bone in our skeleton, we don't want to have an out of bounds access.
-		if (SkeletonBoneIndex >= IndexToBodyNum)
+		if (!SkeletonBoneIndexToBodyIndex.IsValidIndex(SkeletonBoneIndex))
 		{
 			bHasInvalidBoneReference = true;
 			break;
@@ -1857,7 +1856,7 @@ void FAnimNode_RigidBody::InitializeBoneReferences(const FBoneContainer& Require
 				const int32 SkeletonParentBoneIndex = RequiredBones.GetSkeletonIndex(CompactParentIndex);
 
 				// Must check our parent as well for a missing bone.
-				if (SkeletonParentBoneIndex >= IndexToBodyNum)
+				if (!SkeletonBoneIndexToBodyIndex.IsValidIndex(SkeletonParentBoneIndex))
 				{
 					bHasInvalidBoneReference = true;
 					break;
