Animation Memory:  17.7 MB savings in array slack.
* 14.3 MB:  Reserve space in CurvePoseSourceIndices members in AnimNode_BoneMask and AnimNode_LayeredBoneBlend based on the exact amount used.  Many nodes don't use this feature, so the array ends up empty, but in other cases, less than 5% of the reserve is used.
* 3.4 MB:  Eliminate Reserve for QueuedTransitionEvents member in AnimNode_StateMachine.  Reserve is wasteful because it's deleted anyway by a Shrink call in FAnimNode_StateMachine::Update_AnyThread for active state machines, and only ends up sticking around for inactive state machines, where it's irrelevant, so it doesn't affect performance to take this out.

#rnx

diff --git a/Engine/Source/Runtime/AnimGraphRuntime/Private/AnimNodes/AnimNode_LayeredBoneBlend.cpp b/Engine/Source/Runtime/AnimGraphRuntime/Private/AnimNodes/AnimNode_LayeredBoneBlend.cpp
--- a/Engine/Source/Runtime/AnimGraphRuntime/Private/AnimNodes/AnimNode_LayeredBoneBlend.cpp
+++ b/Engine/Source/Runtime/AnimGraphRuntime/Private/AnimNodes/AnimNode_LayeredBoneBlend.cpp
@@ -88,8 +88,9 @@ void FAnimNode_LayeredBoneBlend::UpdateCachedBoneData(const FBoneContainer& Requ
 
 	// Build curve source indices
 	{
+		// Get the original Reserve value
+		int32 OriginalReserve = CurvePoseSourceIndices.Max();
 		CurvePoseSourceIndices.Empty();
-		CurvePoseSourceIndices.Reserve(Skeleton->GetNumCurveMetaData());
 
 		Skeleton->ForEachCurveMetaData([this, &RequiredBones](const FName& InCurveName, const FCurveMetaData& InMetaData)
 		{
@@ -106,6 +107,13 @@ void FAnimNode_LayeredBoneBlend::UpdateCachedBoneData(const FBoneContainer& Requ
 				}
 			}
 		});
+
+		// Shrink afterwards to exactly what was used if the Reserve increased, to save memory.  Eventually the reserve will
+		// stabilize at the maximum number of nodes actually used in practice for this specific anim node.
+		if (CurvePoseSourceIndices.Num() > OriginalReserve)
+		{
+			CurvePoseSourceIndices.Shrink();
+		}
 	}
 
 	RequiredBonesSerialNumber = RequiredBones.GetSerialNumber();
diff --git a/Engine/Source/Runtime/Engine/Private/Animation/AnimNode_StateMachine.cpp b/Engine/Source/Runtime/Engine/Private/Animation/AnimNode_StateMachine.cpp
--- a/Engine/Source/Runtime/Engine/Private/Animation/AnimNode_StateMachine.cpp
+++ b/Engine/Source/Runtime/Engine/Private/Animation/AnimNode_StateMachine.cpp
@@ -265,7 +265,7 @@ void FAnimNode_StateMachine::Initialize_AnyThread(const FAnimationInitializeCont
 			// Reset transition related variables
 			StatesUpdated.Reset();
 			ActiveTransitionArray.Reset();
-			QueuedTransitionEvents.Reset(MaxTransitionsRequests);
+			QueuedTransitionEvents.Reset();
 
 			StateCacheBoneCounters.Reset(Machine->States.Num());
 			StateCacheBoneCounters.AddDefaulted(Machine->States.Num());
@@ -1426,7 +1426,7 @@ void FAnimNode_StateMachine::ClearTransitionEvents(const FName& EventName)
 void FAnimNode_StateMachine::ClearAllTransitionEvents()
 {
 	UE_LOG(LogAnimTransitionRequests, Verbose, TEXT("Clearing all request (Machine %s)"), *GetMachineDescription()->MachineName.ToString());
-	QueuedTransitionEvents.Reset(MaxTransitionsRequests);
+	QueuedTransitionEvents.Reset();
 }
 
 bool FAnimNode_StateMachine::QueryTransitionEvent(const int32 TransitionIndex, const FName& EventName) const
diff --git a/Engine/Source/Runtime/Engine/Public/Animation/NamedValueArray.h b/Engine/Source/Runtime/Engine/Public/Animation/NamedValueArray.h
--- a/Engine/Source/Runtime/Engine/Public/Animation/NamedValueArray.h
+++ b/Engine/Source/Runtime/Engine/Public/Animation/NamedValueArray.h
@@ -131,7 +131,19 @@ struct TNamedValueArray
 	{
 		return Elements.Num();
 	}
-	
+
+	/** @returns the max number of elements reserved in the array */
+	int32 Max() const
+	{
+		return Elements.Max();
+	}
+
+	/** Compacts the memory for the elements based on what was actually used */
+	void Shrink()
+	{
+		return Elements.Shrink();
+	}
+
 protected:
 	// Sort by FName - Note: this is not stable across serialization
 	struct FElementSortPredicate
