PR #9749: Fix in randomplayer wrong swap code when min,max playrate need swap.

diff --git a/Engine/Source/Runtime/AnimGraphRuntime/Private/AnimNodes/AnimNode_RandomPlayer.cpp b/Engine/Source/Runtime/AnimGraphRuntime/Private/AnimNodes/AnimNode_RandomPlayer.cpp
--- a/Engine/Source/Runtime/AnimGraphRuntime/Private/AnimNodes/AnimNode_RandomPlayer.cpp
+++ b/Engine/Source/Runtime/AnimGraphRuntime/Private/AnimNodes/AnimNode_RandomPlayer.cpp
@@ -68,7 +68,7 @@ void FAnimNode_RandomPlayer::Initialize_AnyThread(const FAnimationInitializeCont
 
 		if (Entry->MaxPlayRate < Entry->MinPlayRate)
 		{
-			Swap(Entry->MaxLoopCount, Entry->MinLoopCount);
+			Swap(Entry->MaxPlayRate, Entry->MinPlayRate);
 		}
 
 		Entry->BlendIn.Reset();
