Disabled RBAN (Rigid Body Animation Nodes) on the server.
- Created new cvar: "p.RigidBodyNodeServer" which defaults to 0 (disabled).
- Added new STAT_RigidBodyNodeInitTime_SetupSimulation which tracks the more expensive part of calling FAnimNode_RigidBody::InitPhysics.

#tests Tested in editor and cooked builds (Development and Test configs).

diff --git a/Engine/Source/Runtime/AnimGraphRuntime/Private/BoneControllers/AnimNode_RigidBody.cpp b/Engine/Source/Runtime/AnimGraphRuntime/Private/BoneControllers/AnimNode_RigidBody.cpp
--- a/Engine/Source/Runtime/AnimGraphRuntime/Private/BoneControllers/AnimNode_RigidBody.cpp
+++ b/Engine/Source/Runtime/AnimGraphRuntime/Private/BoneControllers/AnimNode_RigidBody.cpp
@@ -24,6 +24,7 @@
 #include "PhysicsEngine/PhysicsSettings.h"
 #include "Logging/MessageLog.h"
 #include "Logging/LogMacros.h"
+#include "ProfilingDebugging/CountersTrace.h"
 
 #include UE_INLINE_GENERATED_CPP_BY_NAME(AnimNode_RigidBody)
 
@@ -37,6 +38,7 @@ LLM_DEFINE_TAG(Animation_RigidBody);
 #define LOCTEXT_NAMESPACE "ImmediatePhysics"
 
 DEFINE_STAT(STAT_RigidBodyNodeInitTime);
+DEFINE_STAT(STAT_RigidBodyNodeInitTime_SetupSimulation);
 
 CSV_DECLARE_CATEGORY_MODULE_EXTERN(ENGINE_API, Animation);
 
@@ -47,7 +49,11 @@ DECLARE_LOG_CATEGORY_EXTERN(LogRBAN, Log, All);
 #endif
 DEFINE_LOG_CATEGORY(LogRBAN);
 
+TRACE_DECLARE_INT_COUNTER(AnimNode_RigidBody_Enabled, TEXT("FAnimNode/RigidBody/Enabled"));
+TRACE_DECLARE_INT_COUNTER(AnimNode_RigidBody_Disabled, TEXT("FAnimNode/RigidBody/Disabled"));
+
 TAutoConsoleVariable<int32> CVarEnableRigidBodyNode(TEXT("p.RigidBodyNode"), 1, TEXT("Enables/disables the whole rigid body node system. When disabled, avoids all allocations and runtime costs. Can be used to disable RB Nodes on low-end platforms."), ECVF_Scalability);
+TAutoConsoleVariable<int32> CVarEnableRigidBodyNodeServer(TEXT("p.RigidBodyNodeServer"), 0, TEXT("Enables/disables the whole rigid body node system on the Server. When disabled, avoids all allocations and runtime costs."), ECVF_Default);
 TAutoConsoleVariable<int32> CVarEnableRigidBodyNodeSimulation(TEXT("p.RigidBodyNode.EnableSimulation"), 1, TEXT("Runtime Enable/Disable RB Node Simulation for debugging and testing (node is initialized and bodies and constraints are created, even when disabled.)"), ECVF_Default);
 TAutoConsoleVariable<int32> CVarRigidBodyLODThreshold(TEXT("p.RigidBodyLODThreshold"), -1, TEXT("Max LOD that rigid body node is allowed to run on. Provides a global threshold that overrides per-node the LODThreshold property. -1 means no override."), ECVF_Scalability);
 
@@ -1057,10 +1063,15 @@ void FAnimNode_RigidBody::InitPhysics(const UAnimInstance* InAnimInstance)
 		AnimPhysicsMinDeltaTime = 0.f;
 		bSimulateAnimPhysicsAfterReset = false;
 	}
-	
-	bEnabled = UsePhysicsAsset && SkeletalMeshComp->GetAllowRigidBodyAnimNode() && CVarEnableRigidBodyNode.GetValueOnAnyThread() != 0;
+
+	const bool bUseRigidBodyCVar = IsEnableRigidBodyNodeCVarEnabled(InAnimInstance);
+	bEnabled = (UsePhysicsAsset && SkeletalMeshComp->GetAllowRigidBodyAnimNode() && bUseRigidBodyCVar);
 	if(bEnabled)
 	{
+		TRACE_COUNTER_INCREMENT(AnimNode_RigidBody_Enabled);
+
+		SCOPE_CYCLE_COUNTER(STAT_RigidBodyNodeInitTime_SetupSimulation);
+
 		PhysicsSimulation = new ImmediatePhysics::FSimulation();
 
 #if WITH_CHAOS_VISUAL_DEBUGGER
@@ -1294,6 +1305,10 @@ void FAnimNode_RigidBody::InitPhysics(const UAnimInstance* InAnimInstance)
 
 		SolverIterations = UsePhysicsAsset->SolverIterations;
 	}
+	else
+	{
+		TRACE_COUNTER_INCREMENT(AnimNode_RigidBody_Disabled);
+	}
 }
 
 DECLARE_CYCLE_STAT(TEXT("FAnimNode_RigidBody::UpdateWorldGeometry"), STAT_ImmediateUpdateWorldGeometry, STATGROUP_ImmediatePhysics);
@@ -1465,6 +1480,19 @@ void FAnimNode_RigidBody::SetOverridePhysicsAsset(UPhysicsAsset* PhysicsAsset)
 	OverridePhysicsAsset = PhysicsAsset;
 }
 
+bool FAnimNode_RigidBody::IsEnableRigidBodyNodeCVarEnabled(const UAnimInstance* InAnimInstance)
+{
+	// Pick the correct CVar depending on NetMode:
+	//  Client: p.RigidBodyNode
+	//  DedicatedServer: p.RigidBodyNodeServer
+
+	UWorld* World = InAnimInstance->GetWorld();
+	const ENetMode NetMode = World->GetNetMode();
+	const bool bIsDedicatedServer = (NetMode == ENetMode::NM_DedicatedServer);
+	const bool bEnableRigidBodyNodeCVarEnabled = bIsDedicatedServer ? (CVarEnableRigidBodyNodeServer.GetValueOnAnyThread() != 0) : (CVarEnableRigidBodyNode.GetValueOnAnyThread() != 0);
+	return bEnableRigidBodyNodeCVarEnabled;
+}
+
 DECLARE_CYCLE_STAT(TEXT("RigidBody_PreUpdate"), STAT_RigidBody_PreUpdate, STATGROUP_Anim);
 
 void FAnimNode_RigidBody::PreUpdate(const UAnimInstance* InAnimInstance)
diff --git a/Engine/Source/Runtime/AnimGraphRuntime/Public/BoneControllers/AnimNode_RigidBody.h b/Engine/Source/Runtime/AnimGraphRuntime/Public/BoneControllers/AnimNode_RigidBody.h
--- a/Engine/Source/Runtime/AnimGraphRuntime/Public/BoneControllers/AnimNode_RigidBody.h
+++ b/Engine/Source/Runtime/AnimGraphRuntime/Public/BoneControllers/AnimNode_RigidBody.h
@@ -15,6 +15,7 @@ struct FConstraintInstance;
 class FEvent;
 
 extern ANIMGRAPHRUNTIME_API TAutoConsoleVariable<int32> CVarEnableRigidBodyNode;
+extern ANIMGRAPHRUNTIME_API TAutoConsoleVariable<int32> CVarEnableRigidBodyNodeServer;
 extern ANIMGRAPHRUNTIME_API TAutoConsoleVariable<int32> CVarEnableRigidBodyNodeSimulation;
 extern ANIMGRAPHRUNTIME_API TAutoConsoleVariable<int32> CVarRigidBodyLODThreshold;
 
@@ -186,6 +187,9 @@ struct FAnimNode_RigidBody : public FAnimNode_SkeletalControlBase
 
 	UPhysicsAsset* GetPhysicsAsset() const { return UsePhysicsAsset; }
 
+	// Used to check if whether p.RigidBodyNode or p.RigidBodyNodeServer is enabled.
+	ANIMGRAPHRUNTIME_API static bool IsEnableRigidBodyNodeCVarEnabled(const UAnimInstance* InAnimInstance);
+
 public:
 	/** Physics asset to use. If empty use the skeletal mesh's default physics asset in case Default To Skeletal Mesh Physics Asset is set to True. */
 	UPROPERTY(EditAnywhere, Category = Settings)
diff --git a/Engine/Source/Runtime/Engine/Public/Physics/ImmediatePhysics/ImmediatePhysicsStats.h b/Engine/Source/Runtime/Engine/Public/Physics/ImmediatePhysics/ImmediatePhysicsStats.h
--- a/Engine/Source/Runtime/Engine/Public/Physics/ImmediatePhysics/ImmediatePhysicsStats.h
+++ b/Engine/Source/Runtime/Engine/Public/Physics/ImmediatePhysics/ImmediatePhysicsStats.h
@@ -8,3 +8,4 @@ DECLARE_STATS_GROUP(TEXT("Immediate Physics"), STATGROUP_ImmediatePhysics, STATC
 DECLARE_STATS_GROUP(TEXT("Immediate Physics Counters"), STATGROUP_ImmediatePhysicsCounters, STATCAT_Advanced);
 
 DECLARE_CYCLE_STAT_EXTERN(TEXT("RigidBodyNodeInitTime"), STAT_RigidBodyNodeInitTime, STATGROUP_ImmediatePhysics,);
+DECLARE_CYCLE_STAT_EXTERN(TEXT("RigidBodyNodeInitTime_SetupSimulation"), STAT_RigidBodyNodeInitTime_SetupSimulation, STATGROUP_ImmediatePhysics,);
