Add possibility to expose the LOD thresholds as pins

for the following nodes:
Control Rig
Rig Logic
Skeletal Control Base (RBAN)

The plan is to use this on AnimBPs that are shared across different characters presets so that we can granularily control compute on the LOD levels for each.

diff --git a/Engine/Plugins/Animation/ControlRig/Source/ControlRig/Public/AnimNode_ControlRig.h b/Engine/Plugins/Animation/ControlRig/Source/ControlRig/Public/AnimNode_ControlRig.h
--- a/Engine/Plugins/Animation/ControlRig/Source/ControlRig/Public/AnimNode_ControlRig.h
+++ b/Engine/Plugins/Animation/ControlRig/Source/ControlRig/Public/AnimNode_ControlRig.h
@@ -109,7 +109,7 @@ private:
 	 * when the component LOD becomes 3, it will stop update/evaluate
 	 * currently transition would be issue and that has to be re-visited
 	 */
-	UPROPERTY(EditAnywhere, Category = Performance, meta = (DisplayName = "LOD Threshold"))
+	UPROPERTY(EditAnywhere, Category = Performance, meta = (PinHiddenByDefault, DisplayName = "LOD Threshold"))
 	int32 LODThreshold;
 
 private:
diff --git a/Engine/Plugins/Animation/RigLogic/Source/RigLogicModule/Public/AnimNode_RigLogic.h b/Engine/Plugins/Animation/RigLogic/Source/RigLogicModule/Public/AnimNode_RigLogic.h
--- a/Engine/Plugins/Animation/RigLogic/Source/RigLogicModule/Public/AnimNode_RigLogic.h
+++ b/Engine/Plugins/Animation/RigLogic/Source/RigLogicModule/Public/AnimNode_RigLogic.h
@@ -55,7 +55,7 @@ private:
 	 * For example if you have the threshold set to 2, it will evaluate until including LOD 2 (based on 0 index). In case the LOD level gets set to 3, it will stop evaluating the post-process AnimBP.
 	 * Setting it to -1 will always evaluate it and disable LODing.
 	 */
-	UPROPERTY(EditAnywhere, Category = RigLogic)
+	UPROPERTY(EditAnywhere, Category = RigLogic, meta = (PinHiddenByDefault))
 	int32 LODThreshold = INDEX_NONE;
 
 private:
diff --git a/Engine/Source/Runtime/AnimGraphRuntime/Public/BoneControllers/AnimNode_SkeletalControlBase.h b/Engine/Source/Runtime/AnimGraphRuntime/Public/BoneControllers/AnimNode_SkeletalControlBase.h
--- a/Engine/Source/Runtime/AnimGraphRuntime/Public/BoneControllers/AnimNode_SkeletalControlBase.h
+++ b/Engine/Source/Runtime/AnimGraphRuntime/Public/BoneControllers/AnimNode_SkeletalControlBase.h
@@ -32,7 +32,7 @@ struct FAnimNode_SkeletalControlBase : public FAnimNode_Base
 	* when the component LOD becomes 3, it will stop update/evaluate
 	* currently transition would be issue and that has to be re-visited
 	*/
-	UPROPERTY(EditAnywhere, BlueprintReadWrite, Category = Performance, meta = (DisplayName = "LOD Threshold"))
+	UPROPERTY(EditAnywhere, BlueprintReadWrite, Category = Performance, meta = (PinHiddenByDefault, DisplayName = "LOD Threshold"))
 	int32 LODThreshold;
 
 	UPROPERTY(Transient)
