Adding "soft IK" to LegIK anim node.

diff --git a/Engine/Source/Runtime/AnimGraphRuntime/Private/BoneControllers/AnimNode_LegIK.cpp b/Engine/Source/Runtime/AnimGraphRuntime/Private/BoneControllers/AnimNode_LegIK.cpp
--- a/Engine/Source/Runtime/AnimGraphRuntime/Private/BoneControllers/AnimNode_LegIK.cpp
+++ b/Engine/Source/Runtime/AnimGraphRuntime/Private/BoneControllers/AnimNode_LegIK.cpp
@@ -1,11 +1,13 @@
 // Copyright Epic Games, Inc. All Rights Reserved.
 
 #include "BoneControllers/AnimNode_LegIK.h"
+
 #include "Components/SkeletalMeshComponent.h"
 #include "DrawDebugHelpers.h"
 #include "Engine/Engine.h"
 #include "EngineGlobals.h"
 #include "Animation/AnimInstanceProxy.h"
+#include "SoftIK.h"
 
 #include UE_INLINE_GENERATED_CPP_BY_NAME(AnimNode_LegIK)
 
@@ -29,6 +31,8 @@ FAnimNode_LegIK::FAnimNode_LegIK()
 {
 	ReachPrecision = 0.01f;
 	MaxIterations = 12;
+	SoftPercentLength = 1.0f;
+	SoftAlpha = 1.0f;
 }
 
 void FAnimNode_LegIK::GatherDebugData(FNodeDebugData& DebugData)
@@ -217,7 +221,7 @@ void FIKChain::InitializeFromLegData(FAnimLegIKData& InLegData, FAnimInstancePro
 		Links.Init(FIKChainLink(), InLegData.NumBones);
 	}
 	
-	MaximumReach = 0.0;
+	TotalChainLength = 0.0;
 
 	check(InLegData.NumBones > 1);
 	for (int32 Index = 0; Index < InLegData.NumBones - 1; Index++)
@@ -230,7 +234,7 @@ void FIKChain::InitializeFromLegData(FAnimLegIKData& InLegData, FAnimInstancePro
 		Link.Location = BoneLocation;
 		Link.Length = BoneLength;
 
-		MaximumReach += BoneLength;
+		TotalChainLength += BoneLength;
 	}
 
 	// Add root bone last
@@ -260,7 +264,12 @@ void FIKChain::InitializeFromLegData(FAnimLegIKData& InLegData, FAnimInstancePro
 
 TAutoConsoleVariable<int32> CVarAnimLegIKTwoBone(TEXT("a.AnimNode.LegIK.EnableTwoBone"), 1, TEXT("Enable Two Bone Code Path."));
 
-void FIKChain::ReachTarget(const FVector& InTargetLocation, double InReachPrecision, int32 InMaxIterations)
+void FIKChain::ReachTarget(
+	const FVector& InTargetLocation,
+	double InReachPrecision,
+	int32 InMaxIterations,
+	float SoftPercentLength,
+	float SoftAlpha)
 {
 	if (!bInitialized)
 	{
@@ -269,8 +278,18 @@ void FIKChain::ReachTarget(const FVector& InTargetLocation, double InReachPrecis
 
 	const FVector RootLocation = Links.Last().Location;
 
+	// Optionally soften the target location to prevent knee popping
+	FVector FinalTargetLocation = InTargetLocation;
+	const bool bUsingSoftIK = SoftPercentLength < 1.0f && SoftAlpha > 0.f;
+	if (bUsingSoftIK)
+	{
+		AnimationCore::SoftenIKEffectorPosition(RootLocation, TotalChainLength, SoftPercentLength, SoftAlpha, FinalTargetLocation);
+	}
+
 	// If we can't reach, we just go in a straight line towards the target,
-	if ((NumLinks <= 2) || (FVector::DistSquared(RootLocation, InTargetLocation) >= FMath::Square(GetMaximumReach())))
+	const bool bTargetIsReachable = FVector::DistSquared(RootLocation, InTargetLocation) < FMath::Square(GetMaximumReach());
+	const bool bHasTwoOrFewerLinks = NumLinks <= 2;
+	if (bHasTwoOrFewerLinks || (!bTargetIsReachable && !bUsingSoftIK))
 	{
 		const FVector Direction = (InTargetLocation - RootLocation).GetSafeNormal();
 		OrientAllLinksToDirection(Direction);
@@ -278,12 +297,12 @@ void FIKChain::ReachTarget(const FVector& InTargetLocation, double InReachPrecis
 	// Two Bones, we can figure out solution instantly
 	else if (NumLinks == 3 && (CVarAnimLegIKTwoBone.GetValueOnAnyThread() == 1))
 	{
-		SolveTwoBoneIK(InTargetLocation);
+		SolveTwoBoneIK(FinalTargetLocation);
 	}
 	// Do iterative approach based on FABRIK
 	else
 	{
-		SolveFABRIK(InTargetLocation, InReachPrecision, InMaxIterations);
+		SolveFABRIK(FinalTargetLocation, InReachPrecision, InMaxIterations);
 	}
 }
 
@@ -408,7 +427,7 @@ bool FAnimNode_LegIK::DoLegReachIK(FAnimLegIKData& InLegData)
 	if (bNeedsSolver)
 	{
 		const int32 MaxIterationsOverride = CVarAnimLegIKMaxIterations.GetValueOnAnyThread() > 0 ? CVarAnimLegIKMaxIterations.GetValueOnAnyThread() : MaxIterations;
-		IKChain.ReachTarget(FootIKLocation, ReachPrecision, MaxIterationsOverride);
+		IKChain.ReachTarget(FootIKLocation, ReachPrecision, MaxIterationsOverride, SoftPercentLength, SoftAlpha);
 	}
 
 	if (bHasTwistOffset)
diff --git a/Engine/Source/Runtime/AnimGraphRuntime/Public/BoneControllers/AnimNode_LegIK.h b/Engine/Source/Runtime/AnimGraphRuntime/Public/BoneControllers/AnimNode_LegIK.h
--- a/Engine/Source/Runtime/AnimGraphRuntime/Public/BoneControllers/AnimNode_LegIK.h
+++ b/Engine/Source/Runtime/AnimGraphRuntime/Public/BoneControllers/AnimNode_LegIK.h
@@ -56,7 +56,7 @@ public:
 
 private:
 	FAnimInstanceProxy* MyAnimInstanceProxy;
-	double MaximumReach;
+	double TotalChainLength;
 	int32 NumLinks;
 	FVector HingeRotationAxis;
 	bool bEnableRotationLimit;
@@ -65,7 +65,7 @@ private:
 public:
 	FIKChain()
 		: MyAnimInstanceProxy(nullptr)
-		, MaximumReach(0.0)
+		, TotalChainLength(0.0)
 		, NumLinks(INDEX_NONE)
 		, HingeRotationAxis(FVector::ZeroVector)
 		, bEnableRotationLimit(false)
@@ -73,12 +73,17 @@ public:
 	{}
 
 	void InitializeFromLegData(FAnimLegIKData& InLegData, FAnimInstanceProxy* InAnimInstanceProxy);
-	void ReachTarget(const FVector& InTargetLocation, double InReachPrecision, int32 InMaxIterations);
+	void ReachTarget(
+		const FVector& InTargetLocation,
+		double InReachPrecision,
+		int32 InMaxIterations,
+		float SoftPercentLength,
+		float SoftAlpha);
 	void ApplyTwistOffset(const float InTwistOffsetDegrees);
 
 	double GetMaximumReach() const
 	{
-		return MaximumReach;
+		return TotalChainLength;
 	}
 
 private:
@@ -191,6 +196,16 @@ struct FAnimNode_LegIK : public FAnimNode_SkeletalControlBase
 	UPROPERTY(EditAnywhere, Category = "Settings")
 	int32 MaxIterations;
 
+	/** Default is 1.0 (off). Range is 0.1 to 1.0. When set to a value less than 1, will "softly" approach full extension starting when the effector
+	 * distance from the root of the chain is greater than this percent length of the bone chain. Typical values are around 0.97.
+	 * This is useful for preventing the knee from "popping" when approaching full extension. */
+	UPROPERTY(EditAnywhere, Category = "Settings", meta = (PinHiddenByDefault, UIMin = "0.01", UIMax = "1", ClampMin = "0.01", ClampMax = "1"))
+	float SoftPercentLength;
+
+	/** Default is 1.0 (full). Range is 0 to 1. Blends the effect of the "softness" on/off. */
+	UPROPERTY(EditAnywhere, Category = "Settings", meta = (PinHiddenByDefault, UIMin = "0.0", UIMax = "1", ClampMin = "0.0", ClampMax = "1"))
+	float SoftAlpha;
+
 	UPROPERTY(EditAnywhere, Category = "Settings")
 	TArray<FAnimLegIKDefinition> LegsDefinition;
 
diff --git a/Engine/Source/Runtime/AnimationCore/Private/SoftIK.cpp b/Engine/Source/Runtime/AnimationCore/Private/SoftIK.cpp
new file mode 100644
--- /dev/null
+++ b/Engine/Source/Runtime/AnimationCore/Private/SoftIK.cpp
@@ -0,0 +1,60 @@
+﻿// Copyright Epic Games, Inc. All Rights Reserved.
+
+#include "SoftIK.h"
+
+namespace AnimationCore
+{
+	void SoftenIKEffectorPosition(
+		const FVector& RootLocation,
+		const float TotalChainLength,
+		const float SoftLengthPercent,
+		const float Alpha,
+		FVector& InOutEffectorPosition)
+	{
+		if (FMath::Abs(1.0f - SoftLengthPercent) < KINDA_SMALL_NUMBER)
+		{
+			// soft length near zero
+			return;
+		}
+
+		if (TotalChainLength <= KINDA_SMALL_NUMBER)
+		{
+			// chain length near zero
+			return;
+		}
+		
+		// get vector from root of chain to effector (pre adjusted)
+		FVector StartToEffector = InOutEffectorPosition - RootLocation;
+		float CurrentLength;
+		StartToEffector.ToDirectionAndLength(StartToEffector, CurrentLength);
+
+		// convert percentage to distance
+		const float SoftDistance = TotalChainLength * (1.0f - FMath::Min(1.0f, SoftLengthPercent));
+		const float HardLength = TotalChainLength - SoftDistance;
+		const float CurrentDelta = CurrentLength - HardLength;
+		if (CurrentDelta <= KINDA_SMALL_NUMBER || SoftDistance <= KINDA_SMALL_NUMBER)
+		{
+			// not in the soft zone
+			return;
+		}
+
+		// calculate the "softened" length of the effector
+		const float PercentIntoSoftLength = CurrentDelta / SoftDistance;
+		const float SoftenedLength = HardLength + SoftDistance * (1.0 - FMath::Exp(-PercentIntoSoftLength));
+
+		// apply the new effector location
+		float FinalAlpha = FMath::Clamp(Alpha, 0.0, 1.0f);
+		if (FinalAlpha < 1.0f)
+		{
+			// alpha blend the softness (optional)
+			const float MaxLength = FMath::Min(CurrentLength, TotalChainLength);
+			const float AlphaBlendedLength = FMath::Lerp(MaxLength, SoftenedLength, Alpha);
+			InOutEffectorPosition = RootLocation + StartToEffector * AlphaBlendedLength;
+		}
+		else
+		{
+			// use the soft position directly
+			InOutEffectorPosition = RootLocation + StartToEffector * SoftenedLength;
+		}
+	}
+}
diff --git a/Engine/Source/Runtime/AnimationCore/Public/SoftIK.h b/Engine/Source/Runtime/AnimationCore/Public/SoftIK.h
new file mode 100644
--- /dev/null
+++ b/Engine/Source/Runtime/AnimationCore/Public/SoftIK.h
@@ -0,0 +1,30 @@
+﻿// Copyright Epic Games, Inc. All Rights Reserved.
+
+#pragma once
+
+#include "CoreMinimal.h"
+#include "Math/MathFwd.h"
+
+namespace AnimationCore
+{
+	/**
+	 * Soft IK
+	 *
+	 * This modifies the effector position for a chain of joints being controlled by IK by making the effector softly approach full extension.
+	 * It takes the desired effector position and adjusts it according to an exponential falloff as the effector approach the full length of the chain.
+	 *
+	 * The technique used is described here: https://web.archive.org/web/20160610183037/https://softimageblog.com/archives/108
+	 *
+	 * @param RootLocation			The location of the bone at the root of the chain
+	 * @param TotalChainLength		The total length of the chain from root to tip, along the bones (not straight line).
+	 * @param SoftLengthPercent		The percentage of the chain length to begin softening the effector motion (typically set to 0.97)
+	 * @param Alpha					The amount of softness to apply (0 is none, 1 is full amount). Default is 1.
+	 * @param InOutEffectorPosition	The input position of the end effector. This is modified by this function to "soften" the IK.
+	 */
+	ANIMATIONCORE_API void SoftenIKEffectorPosition(
+		const FVector& RootLocation,
+		const float TotalChainLength,
+		const float SoftLengthPercent,
+		const float Alpha,
+		FVector& InOutEffectorPosition);
+}
