Exposed LOD threshold as a pin for the pose driver node

diff --git a/Engine/Source/Runtime/AnimGraphRuntime/Public/AnimNodes/AnimNode_PoseDriver.h b/Engine/Source/Runtime/AnimGraphRuntime/Public/AnimNodes/AnimNode_PoseDriver.h
--- a/Engine/Source/Runtime/AnimGraphRuntime/Public/AnimNodes/AnimNode_PoseDriver.h
+++ b/Engine/Source/Runtime/AnimGraphRuntime/Public/AnimNodes/AnimNode_PoseDriver.h
@@ -226,7 +226,7 @@ struct FAnimNode_PoseDriver : public FAnimNode_PoseHandler
 	 * when the component LOD becomes 3, it will stop update/evaluate
 	 * currently transition would be issue and that has to be re-visited
 	 */
-	UPROPERTY(EditAnywhere, Category = Performance, meta = (DisplayName = "LOD Threshold"))
+	UPROPERTY(EditAnywhere, Category = Performance, meta = (PinHiddenByDefault, DisplayName = "LOD Threshold"))
 	int32 LODThreshold;
 
 	// FAnimNode_Base interface
