Fix "weighted moving average" mode in modify curve node

The curve refactor missed this particular case - the weighted moving average mode is supposed to ignore the input values themselves and treat them like a mask on the incoming curves.

diff --git a/Engine/Source/Runtime/AnimGraphRuntime/Private/AnimNodes/AnimNode_ModifyCurve.cpp b/Engine/Source/Runtime/AnimGraphRuntime/Private/AnimNodes/AnimNode_ModifyCurve.cpp
--- a/Engine/Source/Runtime/AnimGraphRuntime/Private/AnimNodes/AnimNode_ModifyCurve.cpp
+++ b/Engine/Source/Runtime/AnimGraphRuntime/Private/AnimNodes/AnimNode_ModifyCurve.cpp
@@ -4,6 +4,7 @@
 #include "AnimationRuntime.h"
 #include "Animation/AnimInstanceProxy.h"
 #include "Animation/AnimCurveTypes.h"
+#include "Animation/AnimCurveUtils.h"
 #include "Animation/AnimStats.h"
 
 #include UE_INLINE_GENERATED_CPP_BY_NAME(AnimNode_ModifyCurve)
@@ -44,41 +45,79 @@ void FAnimNode_ModifyCurve::Evaluate_AnyThread(FPoseContext& Output)
 
 	Output = SourceData;
 
-	// Build internal curves from array & map
-	FBlendedCurve Curve;
-	Curve.Reserve(CurveNames.Num());
-	for(int32 CurveIndex = 0; CurveIndex < CurveNames.Num(); ++CurveIndex)
+	if (ApplyMode == EModifyCurveApplyMode::WeightedMovingAverage)
 	{
-		Curve.Add(CurveNames[CurveIndex], CurveValues[CurveIndex]);
-	}
+		// WMA acts as a filter on incoming curves - it ignores the values that are provided on pins and just uses those
+		// curves as a mask to use when applying the filter
 
-	if(CurveMap.Num() > 0)
-	{
-		FBlendedCurve MapCurve;
-		MapCurve.Reserve(CurveNames.Num());
-		for(const TPair<FName, float>& NameValuePair : CurveMap)
+		UE::Anim::FCurveFilter Filter;
+		Filter.SetFilterMode(UE::Anim::ECurveFilterMode::AllowOnlyFiltered);
+		Filter.Reserve(CurveNames.Num());
+		for(int32 CurveIndex = 0; CurveIndex < CurveNames.Num(); ++CurveIndex)
 		{
-			MapCurve.Add(NameValuePair.Key, NameValuePair.Value);
+			Filter.Add(CurveNames[CurveIndex], UE::Anim::ECurveFilterFlags::Filtered);
 		}
 
-		// Combine pin & map curves in case they overlap
-		Curve.Combine(MapCurve);
-	}
+		if(CurveMap.Num() > 0)
+		{
+			UE::Anim::FCurveFilter MapFilter;
+			MapFilter.Reserve(CurveNames.Num());
+			for(const TPair<FName, float>& NameValuePair : CurveMap)
+			{
+				MapFilter.Add(NameValuePair.Key, UE::Anim::ECurveFilterFlags::Filtered);
+			}
 
-	if (ApplyMode == EModifyCurveApplyMode::WeightedMovingAverage)
-	{
-		UE::Anim::FNamedValueArrayUtils::Union(Output.Curve, Curve, LastCurve,
-			[this](UE::Anim::FCurveElement& InOutResult, const UE::Anim::FCurveElement& InCurveElement, const UE::Anim::FCurveElement& InLastCurveElement, UE::Anim::ENamedValueUnionFlags InFlags)
+			// Combine pin & map curves in case they overlap
+			UE::Anim::FNamedValueArrayUtils::Union(Filter, MapFilter,
+				[](UE::Anim::FCurveFilterElement& InOutElement, const UE::Anim::FCurveFilterElement& InElement, UE::Anim::ENamedValueUnionFlags InFlags)
+				{
+					if(EnumHasAnyFlags(InFlags, UE::Anim::ENamedValueUnionFlags::ValidArg1))
+					{
+						InOutElement.Flags = InElement.Flags;
+					}
+				});
+		}
+
+		// Mask off last curves by pin-exposed curves
+		UE::Anim::FCurveUtils::Filter(LastCurve, Filter);
+
+		// Perform WMA on output
+		UE::Anim::FNamedValueArrayUtils::Union(Output.Curve, LastCurve,
+			[this](UE::Anim::FCurveElement& InOutResult, const UE::Anim::FCurveElement& InCurveElement, UE::Anim::ENamedValueUnionFlags InFlags)
 			{
 				// Only apply curves that we are overriding
 				if(EnumHasAnyFlags(InFlags, UE::Anim::ENamedValueUnionFlags::ValidArg1))
 				{
-					InOutResult.Value = ProcessCurveWMAOperation(InCurveElement.Value, InLastCurveElement.Value);
+					InOutResult.Value = ProcessCurveWMAOperation(InOutResult.Value, InCurveElement.Value);
 				}
 			});
+
+		// Copy current to last
+		LastCurve.CopyFrom(Output.Curve);
 	}
 	else
 	{
+		// Build internal curves from array & map
+		FBlendedCurve Curve;
+		Curve.Reserve(CurveNames.Num());
+		for(int32 CurveIndex = 0; CurveIndex < CurveNames.Num(); ++CurveIndex)
+		{
+			Curve.Add(CurveNames[CurveIndex], CurveValues[CurveIndex]);
+		}
+
+		if(CurveMap.Num() > 0)
+		{
+			FBlendedCurve MapCurve;
+			MapCurve.Reserve(CurveNames.Num());
+			for(const TPair<FName, float>& NameValuePair : CurveMap)
+			{
+				MapCurve.Add(NameValuePair.Key, NameValuePair.Value);
+			}
+
+			// Combine pin & map curves in case they overlap
+			Curve.Combine(MapCurve);
+		}
+		
 		UE::Anim::FNamedValueArrayUtils::Union(Output.Curve, Curve,
 			[this](UE::Anim::FCurveElement& InOutResult, const UE::Anim::FCurveElement& InCurveElement, UE::Anim::ENamedValueUnionFlags InFlags)
 			{
@@ -89,8 +128,6 @@ void FAnimNode_ModifyCurve::Evaluate_AnyThread(FPoseContext& Output)
 				}
 			});
 	}
-
-	LastCurve.CopyFrom(Curve);
 }
 
 float FAnimNode_ModifyCurve::ProcessCurveOperation(float CurrentValue, float NewValue) const
