Fix GC crash with Sequencer anim instance proxies

Proxy was not reporting its references to GC

diff --git a/Engine/Source/Runtime/AnimGraphRuntime/Private/AnimSequencerInstanceProxy.cpp b/Engine/Source/Runtime/AnimGraphRuntime/Private/AnimSequencerInstanceProxy.cpp
--- a/Engine/Source/Runtime/AnimGraphRuntime/Private/AnimSequencerInstanceProxy.cpp
+++ b/Engine/Source/Runtime/AnimGraphRuntime/Private/AnimSequencerInstanceProxy.cpp
@@ -114,6 +114,25 @@ void FAnimSequencerInstanceProxy::ConstructNodes()
 
 }
 
+void FAnimSequencerInstanceProxy::AddReferencedObjects(UAnimInstance* InAnimInstance, FReferenceCollector& Collector)
+{
+	Super::AddReferencedObjects(InAnimInstance, Collector);
+
+	for (const TPair<uint32, FSequencerPlayerBase*>& IndexPlayerPair : SequencerToPlayerMap)
+	{
+		if(IndexPlayerPair.Value->IsOfType<FSequencerPlayerAnimSequence>())
+		{
+			FSequencerPlayerAnimSequence* SequencerPlayerAnimSequence = static_cast<FSequencerPlayerAnimSequence*>(IndexPlayerPair.Value);
+			Collector.AddPropertyReferencesWithStructARO(FAnimNode_SequenceEvaluator_Standalone::StaticStruct(), &SequencerPlayerAnimSequence->PlayerNode);
+		}
+	}
+
+	for (const TPair<uint32, FAnimNode_Mirror_Standalone*>& IndexMirrorPair : SequencerToMirrorMap)
+	{
+		Collector.AddPropertyReferencesWithStructARO(FAnimNode_Mirror_Standalone::StaticStruct(), IndexMirrorPair.Value);
+	}
+}
+
 void FAnimSequencerInstanceProxy::ClearSequencePlayerAndMirrorMaps()
 {
 	for (TPair<uint32, FSequencerPlayerBase*>& Iter : SequencerToPlayerMap)
diff --git a/Engine/Source/Runtime/AnimGraphRuntime/Public/AnimSequencerInstanceProxy.h b/Engine/Source/Runtime/AnimGraphRuntime/Public/AnimSequencerInstanceProxy.h
--- a/Engine/Source/Runtime/AnimGraphRuntime/Public/AnimSequencerInstanceProxy.h
+++ b/Engine/Source/Runtime/AnimGraphRuntime/Public/AnimSequencerInstanceProxy.h
@@ -157,6 +157,8 @@ public:
 	/** Construct and link the base part of the blend tree */
 	ANIMGRAPHRUNTIME_API virtual void ConstructNodes();
 
+	ANIMGRAPHRUNTIME_API virtual void AddReferencedObjects(UAnimInstance* InAnimInstance, FReferenceCollector& Collector) override;
+
 protected:
 
 	ANIMGRAPHRUNTIME_API void UpdateAnimTrack(UAnimSequenceBase* InAnimSequence, uint32 SequenceId, const TOptional<FRootMotionOverride>& RootMomtionOverride, TOptional<float> InFromPosition, float InToPosition, float Weight, bool bFireNotifies, UMirrorDataTable* InMirrorDataTable);
