Prevent pop in soft IK when idling in Leg IK node.

diff --git a/Engine/Source/Runtime/AnimGraphRuntime/Private/BoneControllers/AnimNode_LegIK.cpp b/Engine/Source/Runtime/AnimGraphRuntime/Private/BoneControllers/AnimNode_LegIK.cpp
--- a/Engine/Source/Runtime/AnimGraphRuntime/Private/BoneControllers/AnimNode_LegIK.cpp
+++ b/Engine/Source/Runtime/AnimGraphRuntime/Private/BoneControllers/AnimNode_LegIK.cpp
@@ -413,9 +413,11 @@ bool FAnimNode_LegIK::DoLegReachIK(FAnimLegIKData& InLegData)
 	// The solver is needed if:
 	//	- Our FK foot is not at the IK goal.
 	//	- We're applying a rotation limit.
+	//  - We're using Soft IK (even if foot is at goal, it may be bent by the soft IK if limb is fully extended)
+	const bool bUsingSoftIK = SoftPercentLength < 1.0f && SoftAlpha > 0.f;
 	const bool bFootAtGoal = FootFKLocation.Equals(FootIKLocation, ReachPrecision);
 	const bool bUsingRotationLimit = InLegData.LegDefPtr->bEnableRotationLimit;
-	const bool bNeedsSolver = !bFootAtGoal || bUsingRotationLimit;
+	const bool bNeedsSolver = !bFootAtGoal || bUsingRotationLimit || bUsingSoftIK;
 	if (!bNeedsSolver && !bHasTwistOffset)
 	{
 		return false;
