Add Per axis alpha parameter to hand IK retarget node to allow to apply the resulting translation with custom weights for each axis

diff --git a/Engine/Source/Runtime/AnimGraphRuntime/Private/BoneControllers/AnimNode_HandIKRetargeting.cpp b/Engine/Source/Runtime/AnimGraphRuntime/Private/BoneControllers/AnimNode_HandIKRetargeting.cpp
--- a/Engine/Source/Runtime/AnimGraphRuntime/Private/BoneControllers/AnimNode_HandIKRetargeting.cpp
+++ b/Engine/Source/Runtime/AnimGraphRuntime/Private/BoneControllers/AnimNode_HandIKRetargeting.cpp
@@ -14,7 +14,8 @@ static TAutoConsoleVariable<bool> CVarAnimHandIKRetargetingEnable(TEXT("a.AnimNo
 // FAnimNode_HandIKRetargeting
 
 FAnimNode_HandIKRetargeting::FAnimNode_HandIKRetargeting()
-: HandFKWeight(0.5f)
+: PerAxisAlpha(1.f,1.f,1.f)
+, HandFKWeight(0.5f)
 {
 }
 
@@ -47,6 +48,12 @@ void FAnimNode_HandIKRetargeting::EvaluateSkeletalControl_AnyThread(FComponentSp
 		return;
 	}
 #endif 
+	// Early exit if alpha is zero
+	if (PerAxisAlpha.IsNearlyZero())
+	{
+		return;
+	}
+
 	checkSlow(OutBoneTransforms.Num() == 0);
 
 	const FBoneContainer& BoneContainer = Output.Pose.GetPose().GetBoneContainer();
@@ -87,7 +94,7 @@ void FAnimNode_HandIKRetargeting::EvaluateSkeletalControl_AnyThread(FComponentSp
 		IKLocation = FMath::Lerp<FVector>(LeftHandIKLocation, RightHandIKLocation, HandFKWeight);
 	}
 	
-	const FVector IK_To_FK_Translation = FKLocation - IKLocation;
+	const FVector IK_To_FK_Translation = (FKLocation - IKLocation) * PerAxisAlpha;
 
 	// If we're not translating, don't send any bones to update.
 	if (!IK_To_FK_Translation.IsNearlyZero())
diff --git a/Engine/Source/Runtime/AnimGraphRuntime/Public/BoneControllers/AnimNode_HandIKRetargeting.h b/Engine/Source/Runtime/AnimGraphRuntime/Public/BoneControllers/AnimNode_HandIKRetargeting.h
--- a/Engine/Source/Runtime/AnimGraphRuntime/Public/BoneControllers/AnimNode_HandIKRetargeting.h
+++ b/Engine/Source/Runtime/AnimGraphRuntime/Public/BoneControllers/AnimNode_HandIKRetargeting.h
@@ -13,7 +13,7 @@ class USkeletalMeshComponent;
 
 /**
  * Node to handle re-targeting of Hand IK bone chain.
- * It looks at position in Mesh Space of Left and Right IK bones, and moves Left and Right IK bones to those.
+ * It looks at position in Mesh Space of Left and Right FK bones, and moves Left and Right IK bones to those.
  * based on HandFKWeight. (0 = favor left hand, 1 = favor right hand, 0.5 = equal weight).
  * This is used so characters of different proportions can handle the same props.
  */
@@ -43,6 +43,10 @@ struct FAnimNode_HandIKRetargeting : public FAnimNode_SkeletalControlBase
 	UPROPERTY(EditAnywhere, Category = "HandIKRetargeting")
 	TArray<FBoneReference> IKBonesToMove;
 
+	// Alpha values per axis to apply on the resulting retargeting translation
+	UPROPERTY(EditAnywhere, BlueprintReadWrite, Category = "HandIKRetargeting", meta = (PinHiddenByDefault))
+	FVector PerAxisAlpha;
+
 	/** Which hand to favor. 0.5 is equal weight for both, 1 = right hand, 0 = left hand. */
 	UPROPERTY(EditAnywhere, BlueprintReadWrite, Category = "HandIKRetargeting", meta = (PinShownByDefault))
 	float HandFKWeight;
