Simplified logic which decides whether RBAN should be enabled or disabled server side.
p.RigidBodyNodeServer has been removed in favor of the already existing "p.RigidBodyNode". We are now setting this CVar to false via DedicatedServerEngine.ini.
Refactored CVarEnableRigidBodyNode to use FAutoConsoleVariableRef instead of TAutoConsoleVariable.
Using bools directly instead of calling (CVar.GetValueOnAnyThread())

#tests Tested in editor and cooked builds (Development and Test configs).

diff --git a/Engine/Source/Runtime/AnimGraphRuntime/Private/BoneControllers/AnimNode_RigidBody.cpp b/Engine/Source/Runtime/AnimGraphRuntime/Private/BoneControllers/AnimNode_RigidBody.cpp
--- a/Engine/Source/Runtime/AnimGraphRuntime/Private/BoneControllers/AnimNode_RigidBody.cpp
+++ b/Engine/Source/Runtime/AnimGraphRuntime/Private/BoneControllers/AnimNode_RigidBody.cpp
@@ -24,7 +24,6 @@
 #include "PhysicsEngine/PhysicsSettings.h"
 #include "Logging/MessageLog.h"
 #include "Logging/LogMacros.h"
-#include "ProfilingDebugging/CountersTrace.h"
 
 #include UE_INLINE_GENERATED_CPP_BY_NAME(AnimNode_RigidBody)
 
@@ -49,11 +48,8 @@ DECLARE_LOG_CATEGORY_EXTERN(LogRBAN, Log, All);
 #endif
 DEFINE_LOG_CATEGORY(LogRBAN);
 
-TRACE_DECLARE_INT_COUNTER(AnimNode_RigidBody_Enabled, TEXT("FAnimNode/RigidBody/Enabled"));
-TRACE_DECLARE_INT_COUNTER(AnimNode_RigidBody_Disabled, TEXT("FAnimNode/RigidBody/Disabled"));
-
-TAutoConsoleVariable<int32> CVarEnableRigidBodyNode(TEXT("p.RigidBodyNode"), 1, TEXT("Enables/disables the whole rigid body node system. When disabled, avoids all allocations and runtime costs. Can be used to disable RB Nodes on low-end platforms."), ECVF_Scalability);
-TAutoConsoleVariable<int32> CVarEnableRigidBodyNodeServer(TEXT("p.RigidBodyNodeServer"), 0, TEXT("Enables/disables the whole rigid body node system on the Server. When disabled, avoids all allocations and runtime costs."), ECVF_Default);
+bool bEnableRigidBodyNode = true;
+FAutoConsoleVariableRef CVarEnableRigidBodyNode(TEXT("p.RigidBodyNode"), bEnableRigidBodyNode, TEXT("Enables/disables the whole rigid body node system. When disabled, avoids all allocations and runtime costs. Can be used to disable RB Nodes on low-end platforms."), ECVF_Scalability);
 TAutoConsoleVariable<int32> CVarEnableRigidBodyNodeSimulation(TEXT("p.RigidBodyNode.EnableSimulation"), 1, TEXT("Runtime Enable/Disable RB Node Simulation for debugging and testing (node is initialized and bodies and constraints are created, even when disabled.)"), ECVF_Default);
 TAutoConsoleVariable<int32> CVarRigidBodyLODThreshold(TEXT("p.RigidBodyLODThreshold"), -1, TEXT("Max LOD that rigid body node is allowed to run on. Provides a global threshold that overrides per-node the LODThreshold property. -1 means no override."), ECVF_Scalability);
 
@@ -1061,13 +1057,10 @@ void FAnimNode_RigidBody::InitPhysics(const UAnimInstance* InAnimInstance)
 		AnimPhysicsMinDeltaTime = 0.f;
 		bSimulateAnimPhysicsAfterReset = false;
 	}
-
-	const bool bUseRigidBodyCVar = IsEnableRigidBodyNodeCVarEnabled(InAnimInstance);
-	bEnabled = (UsePhysicsAsset && SkeletalMeshComp->GetAllowRigidBodyAnimNode() && bUseRigidBodyCVar);
+	
+	bEnabled = (UsePhysicsAsset && bEnableRigidBodyNode && SkeletalMeshComp->GetAllowRigidBodyAnimNode());
 	if(bEnabled)
 	{
-		TRACE_COUNTER_INCREMENT(AnimNode_RigidBody_Enabled);
-
 		SCOPE_CYCLE_COUNTER(STAT_RigidBodyNodeInitTime_SetupSimulation);
 
 		PhysicsSimulation = new ImmediatePhysics::FSimulation();
@@ -1303,10 +1296,6 @@ void FAnimNode_RigidBody::InitPhysics(const UAnimInstance* InAnimInstance)
 
 		SolverIterations = UsePhysicsAsset->SolverIterations;
 	}
-	else
-	{
-		TRACE_COUNTER_INCREMENT(AnimNode_RigidBody_Disabled);
-	}
 }
 
 DECLARE_CYCLE_STAT(TEXT("FAnimNode_RigidBody::UpdateWorldGeometry"), STAT_ImmediateUpdateWorldGeometry, STATGROUP_ImmediatePhysics);
@@ -1478,19 +1467,6 @@ void FAnimNode_RigidBody::SetOverridePhysicsAsset(UPhysicsAsset* PhysicsAsset)
 	OverridePhysicsAsset = PhysicsAsset;
 }
 
-bool FAnimNode_RigidBody::IsEnableRigidBodyNodeCVarEnabled(const UAnimInstance* InAnimInstance)
-{
-	// Pick the correct CVar depending on NetMode:
-	//  Client: p.RigidBodyNode
-	//  DedicatedServer: p.RigidBodyNodeServer
-
-	UWorld* World = InAnimInstance->GetWorld();
-	const ENetMode NetMode = World->GetNetMode();
-	const bool bIsDedicatedServer = (NetMode == ENetMode::NM_DedicatedServer);
-	const bool bEnableRigidBodyNodeCVarEnabled = bIsDedicatedServer ? (CVarEnableRigidBodyNodeServer.GetValueOnAnyThread() != 0) : (CVarEnableRigidBodyNode.GetValueOnAnyThread() != 0);
-	return bEnableRigidBodyNodeCVarEnabled;
-}
-
 DECLARE_CYCLE_STAT(TEXT("RigidBody_PreUpdate"), STAT_RigidBody_PreUpdate, STATGROUP_Anim);
 
 void FAnimNode_RigidBody::PreUpdate(const UAnimInstance* InAnimInstance)
diff --git a/Engine/Source/Runtime/AnimGraphRuntime/Public/BoneControllers/AnimNode_RigidBody.h b/Engine/Source/Runtime/AnimGraphRuntime/Public/BoneControllers/AnimNode_RigidBody.h
--- a/Engine/Source/Runtime/AnimGraphRuntime/Public/BoneControllers/AnimNode_RigidBody.h
+++ b/Engine/Source/Runtime/AnimGraphRuntime/Public/BoneControllers/AnimNode_RigidBody.h
@@ -14,8 +14,8 @@ struct FBodyInstance;
 struct FConstraintInstance;
 class FEvent;
 
-extern ANIMGRAPHRUNTIME_API TAutoConsoleVariable<int32> CVarEnableRigidBodyNode;
-extern ANIMGRAPHRUNTIME_API TAutoConsoleVariable<int32> CVarEnableRigidBodyNodeServer;
+extern ANIMGRAPHRUNTIME_API bool bEnableRigidBodyNode;
+extern ANIMGRAPHRUNTIME_API FAutoConsoleVariableRef CVarEnableRigidBodyNode;
 extern ANIMGRAPHRUNTIME_API TAutoConsoleVariable<int32> CVarEnableRigidBodyNodeSimulation;
 extern ANIMGRAPHRUNTIME_API TAutoConsoleVariable<int32> CVarRigidBodyLODThreshold;
 
@@ -187,9 +187,6 @@ struct FAnimNode_RigidBody : public FAnimNode_SkeletalControlBase
 
 	UPhysicsAsset* GetPhysicsAsset() const { return UsePhysicsAsset; }
 
-	// Used to check if whether p.RigidBodyNode or p.RigidBodyNodeServer is enabled.
-	ANIMGRAPHRUNTIME_API static bool IsEnableRigidBodyNodeCVarEnabled(const UAnimInstance* InAnimInstance);
-
 public:
 	/** Physics asset to use. If empty use the skeletal mesh's default physics asset in case Default To Skeletal Mesh Physics Asset is set to True. */
 	UPROPERTY(EditAnywhere, Category = Settings)
