Fix crash showing debug information for a blend space evaluator with no blendspace

diff --git a/Engine/Source/Runtime/AnimGraphRuntime/Private/AnimNodes/AnimNode_BlendSpaceEvaluator.cpp b/Engine/Source/Runtime/AnimGraphRuntime/Private/AnimNodes/AnimNode_BlendSpaceEvaluator.cpp
--- a/Engine/Source/Runtime/AnimGraphRuntime/Private/AnimNodes/AnimNode_BlendSpaceEvaluator.cpp
+++ b/Engine/Source/Runtime/AnimGraphRuntime/Private/AnimNodes/AnimNode_BlendSpaceEvaluator.cpp
@@ -32,8 +32,12 @@ void FAnimNode_BlendSpaceEvaluator::GatherDebugData(FNodeDebugData& DebugData)
 	DECLARE_SCOPE_HIERARCHICAL_COUNTER_ANIMNODE(GatherDebugData)
 	FString DebugLine = DebugData.GetNodeName(this);
 
-	DebugLine += FString::Printf(TEXT("('%s' Play Time: %.3f)"), *GetBlendSpace()->GetName(), InternalTimeAccumulator);
-	DebugData.AddDebugItem(DebugLine, true);
+	UBlendSpace* CurrentBlendSpace = GetBlendSpace();
+	if (CurrentBlendSpace)
+	{
+		DebugLine += FString::Printf(TEXT("('%s' Play Time: %.3f)"), *CurrentBlendSpace->GetName(), InternalTimeAccumulator);
+		DebugData.AddDebugItem(DebugLine, true);
+	}
 }
 
 float FAnimNode_BlendSpaceEvaluator::GetPlayRate() const
