Improve help text for BlendSpace evaluator, which doesn't behave well in all situations.

diff --git a/Engine/Source/Editor/AnimGraph/Private/AnimGraphNode_BlendSpaceEvaluator.cpp b/Engine/Source/Editor/AnimGraph/Private/AnimGraphNode_BlendSpaceEvaluator.cpp
--- a/Engine/Source/Editor/AnimGraph/Private/AnimGraphNode_BlendSpaceEvaluator.cpp
+++ b/Engine/Source/Editor/AnimGraph/Private/AnimGraphNode_BlendSpaceEvaluator.cpp
@@ -26,6 +26,17 @@ FText UAnimGraphNode_BlendSpaceEvaluator::GetNodeTitle(ENodeTitleType::Type Titl
 	return GetNodeTitleHelper(TitleType, BlendSpacePin, LOCTEXT("PlayerDesc", "Blendspace Evaluator"));
 }
 
+FText UAnimGraphNode_BlendSpaceEvaluator::GetTooltipText() const
+{
+	return LOCTEXT("BlendSpaceEvaluatorTooltip", 
+		"Evaluates a BlendSpace at a specific using a specific time input rather than advancing time "
+		"internally. Typically the playback position of the animation for this node will represent "
+		"something other than time, like jump height. Note that events output from the sequences playing "
+		"and being blended together should not be used. In addition, synchronisation of animations "
+		"will potentially be discontinuous if the blend weights are updated, as the leader/follower changes.");
+}
+
+
 void UAnimGraphNode_BlendSpaceEvaluator::GetMenuActions(FBlueprintActionDatabaseRegistrar& InActionRegistrar) const
 {
 	GetMenuActionsHelper(
diff --git a/Engine/Source/Editor/AnimGraph/Public/AnimGraphNode_BlendSpaceEvaluator.h b/Engine/Source/Editor/AnimGraph/Public/AnimGraphNode_BlendSpaceEvaluator.h
--- a/Engine/Source/Editor/AnimGraph/Public/AnimGraphNode_BlendSpaceEvaluator.h
+++ b/Engine/Source/Editor/AnimGraph/Public/AnimGraphNode_BlendSpaceEvaluator.h
@@ -22,6 +22,7 @@ class UAnimGraphNode_BlendSpaceEvaluator : public UAnimGraphNode_BlendSpaceBase
 	// UEdGraphNode interface
 	virtual FText GetNodeTitle(ENodeTitleType::Type TitleType) const override;
 	virtual void GetMenuActions(FBlueprintActionDatabaseRegistrar& ActionRegistrar) const override;
+	virtual FText GetTooltipText() const override;
 	// End of UEdGraphNode interface
 
 	// UAnimGraphNode_Base interface
diff --git a/Engine/Source/Runtime/AnimGraphRuntime/Public/AnimNodes/AnimNode_BlendSpaceEvaluator.h b/Engine/Source/Runtime/AnimGraphRuntime/Public/AnimNodes/AnimNode_BlendSpaceEvaluator.h
--- a/Engine/Source/Runtime/AnimGraphRuntime/Public/AnimNodes/AnimNode_BlendSpaceEvaluator.h
+++ b/Engine/Source/Runtime/AnimGraphRuntime/Public/AnimNodes/AnimNode_BlendSpaceEvaluator.h
@@ -7,20 +7,32 @@
 #include "AnimNodes/AnimNode_BlendSpacePlayer.h"
 #include "AnimNode_BlendSpaceEvaluator.generated.h"
 
-// Evaluates a point in a blendspace, using a specific time input rather than advancing time internally.
-// Typically the playback position of the animation for this node will represent something other than time, like jump height.
+// Evaluates a BlendSpace at a specific using a specific time input rather than advancing time
+// internally. Typically the playback position of the animation for this node will represent
+// something other than time, like jump height. Note that events output from the sequences playing
+// and being blended together should not be used. In addition, synchronization of animations
+// will potentially be discontinuous if the blend weights are updated, as the leader/follower changes.
 USTRUCT(BlueprintInternalUseOnly)
 struct FAnimNode_BlendSpaceEvaluator : public FAnimNode_BlendSpacePlayer
 {
 	GENERATED_USTRUCT_BODY()
 
 public:
-	/** Normalized time between [0,1]. The actual length of a blendspace is dynamic based on the coordinate, so it is exposed as a normalized value. */
+	/**
+	 * Normalized time between [0,1]. The actual length of a blendspace is dynamic based on the coordinate, 
+	 * so it is exposed as a normalized value. Note that treating this as a "time" value that increases (and wraps)
+	 * will not result in the same output as you would get from using a BlendSpace player. The output events
+	 * may not be as expected, and synchronization will sometimes be discontinuous if the leader/follower 
+	 * animations change as a result of changing the blend weights (even if that is done smoothly).
+	 */
 	UPROPERTY(EditAnywhere, BlueprintReadWrite, Category=Settings, meta=(PinShownByDefault))
 	float NormalizedTime;
 
-	/** If true, teleport to normalized time, does NOT advance time (does not trigger notifies, does not extract Root Motion, etc.)
-	If false, will advance time (will trigger notifies, extract root motion if applicable, etc). */
+	/** 
+	 * If true, teleport to normalized time, does NOT advance time (does not trigger notifies, does not 
+	 * extract Root Motion, etc.). If false, will advance time (will trigger notifies, extract root motion 
+	 * if applicable, etc). 
+	 */
 	UPROPERTY(EditAnywhere, Category = Settings, meta = (PinHiddenByDefault))
 	bool bTeleportToNormalizedTime = true;
 
