Clarify help text around animation syncing. No functional changes.

#changelist validated
#virtualized

diff --git a/Engine/Plugins/Experimental/Chooser/Source/Chooser/Internal/AnimNode_ChooserPlayer.h b/Engine/Plugins/Experimental/Chooser/Source/Chooser/Internal/AnimNode_ChooserPlayer.h
--- a/Engine/Plugins/Experimental/Chooser/Source/Chooser/Internal/AnimNode_ChooserPlayer.h
+++ b/Engine/Plugins/Experimental/Chooser/Source/Chooser/Internal/AnimNode_ChooserPlayer.h
@@ -172,15 +172,19 @@ private:
 	FGraphTraversalCounter UpdateCounter;
 	
 #if WITH_EDITORONLY_DATA
-	// The group name (NAME_None if it is not part of any group)
+	// The group name that we synchronize with (NAME_None if it is not part of any group). Note that
+	// this is the name of the group used to sync the output of this node - it will not force
+	// syncing of animations contained by it.
 	UPROPERTY(EditAnywhere, Category = Sync, meta = (FoldProperty))
 	FName GroupName = NAME_None;
 
-	// The role this player can assume within the group (ignored if GroupIndex is INDEX_NONE)
+	// The role this node can assume within the group (ignored if GroupName is not set). Note
+	// that this is the role of the output of this node, not of animations contained by it.
 	UPROPERTY(EditAnywhere, Category = Sync, meta = (FoldProperty))
 	TEnumAsByte<EAnimGroupRole::Type> GroupRole = EAnimGroupRole::CanBeLeader;
 
-	// How synchronization is determined
+	// How this node will synchronize with other animations. Note that this determines how the output
+	// of this node is used for synchronization, not of animations contained by it.
 	UPROPERTY(EditAnywhere, Category = Sync, meta = (FoldProperty))
 	EAnimSyncMethod Method = EAnimSyncMethod::DoNotSync;
 
diff --git a/Engine/Source/Runtime/AnimGraphRuntime/Public/AnimNodes/AnimNode_BlendSpaceGraphBase.h b/Engine/Source/Runtime/AnimGraphRuntime/Public/AnimNodes/AnimNode_BlendSpaceGraphBase.h
--- a/Engine/Source/Runtime/AnimGraphRuntime/Public/AnimNodes/AnimNode_BlendSpaceGraphBase.h
+++ b/Engine/Source/Runtime/AnimGraphRuntime/Public/AnimNodes/AnimNode_BlendSpaceGraphBase.h
@@ -44,11 +44,15 @@ protected:
 	UPROPERTY(EditAnywhere, BlueprintReadWrite, Category = Coordinates, meta = (PinShownByDefault, AllowPrivateAccess))
 	float Y = 0.0f;
 
-	// The group name that we synchronize with. All nodes employing sync beyond this in the anim graph will implicitly use this sync group.
+	// The group name that we synchronize with (NAME_None if it is not part of any group). Note that
+	// this is the name of the group used to sync the output of this node - it will not force
+	// syncing of animations contained by it. Animations inside this Blend Space have their own
+	// options for syncing.
 	UPROPERTY(EditAnywhere, BlueprintReadWrite, Category = Sync, meta = (NeverAsPin, AllowPrivateAccess))
 	FName GroupName = NAME_None;
 
-	// The role this player can assume within the group
+	// The role this Blend Space can assume within the group (ignored if GroupName is not set). Note
+	// that this is the role of the output of this node, not of animations contained by it.
 	UPROPERTY(EditAnywhere, BlueprintReadWrite, Category = Sync, meta = (NeverAsPin, AllowPrivateAccess))
 	TEnumAsByte<EAnimGroupRole::Type> GroupRole = EAnimGroupRole::CanBeLeader;
 
diff --git a/Engine/Source/Runtime/AnimGraphRuntime/Public/AnimNodes/AnimNode_BlendSpacePlayer.h b/Engine/Source/Runtime/AnimGraphRuntime/Public/AnimNodes/AnimNode_BlendSpacePlayer.h
--- a/Engine/Source/Runtime/AnimGraphRuntime/Public/AnimNodes/AnimNode_BlendSpacePlayer.h
+++ b/Engine/Source/Runtime/AnimGraphRuntime/Public/AnimNodes/AnimNode_BlendSpacePlayer.h
@@ -131,15 +131,19 @@ struct FAnimNode_BlendSpacePlayer : public FAnimNode_BlendSpacePlayerBase
 private:
 
 #if WITH_EDITORONLY_DATA
-	// The group name (NAME_None if it is not part of any group)
+	// The group name that we synchronize with (NAME_None if it is not part of any group). Note that
+	// this is the name of the group used to sync the output of this node - it will not force
+	// syncing of animations contained by it.
 	UPROPERTY(EditAnywhere, Category = Sync, meta = (FoldProperty))
 	FName GroupName = NAME_None;
 
-	// The role this player can assume within the group (ignored if GroupIndex is INDEX_NONE)
+	// The role this node can assume within the group (ignored if GroupName is not set). Note
+	// that this is the role of the output of this node, not of animations contained by it.
 	UPROPERTY(EditAnywhere, Category = Sync, meta = (FoldProperty))
 	TEnumAsByte<EAnimGroupRole::Type> GroupRole = EAnimGroupRole::CanBeLeader;
 
-	// How synchronization is determined
+	// How this node will synchronize with other animations. Note that this determines how the output
+	// of this node is used for synchronization, not of animations contained by it.
 	UPROPERTY(EditAnywhere, Category = Sync, meta = (FoldProperty))
 	EAnimSyncMethod Method = EAnimSyncMethod::DoNotSync;
 
@@ -212,15 +216,20 @@ struct FAnimNode_BlendSpacePlayer_Standalone : public FAnimNode_BlendSpacePlayer
 
 private:
 
-	// The group name (NAME_None if it is not part of any group)
+	// The group name that we synchronize with (NAME_None if it is not part of any group). Note that
+	// this is the name of the group used to sync the output of this node - it will not force
+	// syncing of animations contained by it. Animations inside this Blend Space have their own
+	// options for syncing.
 	UPROPERTY(EditAnywhere, Category = Sync)
 	FName GroupName = NAME_None;
 
-	// The role this player can assume within the group (ignored if GroupIndex is INDEX_NONE)
+	// The role this Blend Space can assume within the group (ignored if GroupName is not set). Note
+	// that this is the role of the output of this node, not of animations contained by it.
 	UPROPERTY(EditAnywhere, Category = Sync)
 	TEnumAsByte<EAnimGroupRole::Type> GroupRole = EAnimGroupRole::CanBeLeader;
 
-	// How synchronization is determined
+	// How this asset will synchronize with other assets. Note that this determines how the output
+	// of this node is used for synchronization, not of animations contained by it.
 	UPROPERTY(EditAnywhere, Category = Sync)
 	EAnimSyncMethod Method = EAnimSyncMethod::DoNotSync;
 
diff --git a/Engine/Source/Runtime/AnimGraphRuntime/Public/AnimNodes/AnimNode_SequenceEvaluator.h b/Engine/Source/Runtime/AnimGraphRuntime/Public/AnimNodes/AnimNode_SequenceEvaluator.h
--- a/Engine/Source/Runtime/AnimGraphRuntime/Public/AnimNodes/AnimNode_SequenceEvaluator.h
+++ b/Engine/Source/Runtime/AnimGraphRuntime/Public/AnimNodes/AnimNode_SequenceEvaluator.h
@@ -105,15 +105,15 @@ private:
 	friend class UAnimGraphNode_SequenceEvaluator;
 
 #if WITH_EDITORONLY_DATA
-	// The group name (NAME_None if it is not part of any group)
+	// The group name that we synchronize with (NAME_None if it is not part of any group). 
 	UPROPERTY(EditAnywhere, Category=Sync, meta=(FoldProperty))
 	FName GroupName = NAME_None;
 
-	// The role this player can assume within the group (ignored if GroupIndex is INDEX_NONE)
+	// The role this node can assume within the group (ignored if GroupName is not set)
 	UPROPERTY(EditAnywhere, Category=Sync, meta=(FoldProperty))
 	TEnumAsByte<EAnimGroupRole::Type> GroupRole = EAnimGroupRole::CanBeLeader;
 
-	// How synchronization is determined
+	// How this node will synchronize with other animations.
 	UPROPERTY(EditAnywhere, Category=Sync, meta=(FoldProperty))
 	EAnimSyncMethod Method = EAnimSyncMethod::DoNotSync;
 
@@ -191,15 +191,15 @@ struct FAnimNode_SequenceEvaluator_Standalone : public FAnimNode_SequenceEvaluat
 	GENERATED_BODY()
 
 private:
-	// The group name (NAME_None if it is not part of any group)
+	// The group name that we synchronize with (NAME_None if it is not part of any group). 
 	UPROPERTY(EditAnywhere, Category=Sync)
 	FName GroupName = NAME_None;
 
-	// The role this player can assume within the group (ignored if GroupIndex is INDEX_NONE)
+	// The role this node can assume within the group (ignored if GroupName is not set)
 	UPROPERTY(EditAnywhere, Category=Sync)
 	TEnumAsByte<EAnimGroupRole::Type> GroupRole = EAnimGroupRole::CanBeLeader;
 
-	// How synchronization is determined
+	// How this node will synchronize with other animations.
 	UPROPERTY(EditAnywhere, Category=Sync)
 	EAnimSyncMethod Method = EAnimSyncMethod::DoNotSync;
 
diff --git a/Engine/Source/Runtime/AnimGraphRuntime/Public/AnimNodes/AnimNode_Sync.h b/Engine/Source/Runtime/AnimGraphRuntime/Public/AnimNodes/AnimNode_Sync.h
--- a/Engine/Source/Runtime/AnimGraphRuntime/Public/AnimNodes/AnimNode_Sync.h
+++ b/Engine/Source/Runtime/AnimGraphRuntime/Public/AnimNodes/AnimNode_Sync.h
@@ -22,11 +22,11 @@ private:
 	UPROPERTY(EditAnywhere, Category = Links)
 	FPoseLink Source;
 
-	// The group name that we synchronize with. All nodes employing sync beyond this in the anim graph will implicitly use this sync group.
+	// The group name that we synchronize with (NAME_None if it is not part of any group). 
 	UPROPERTY(EditAnywhere, Category = Settings)
 	FName GroupName;
 
-	// The role this player can assume within the group
+	// The role this animation can assume within the group (ignored if GroupName is not set)
 	UPROPERTY(EditAnywhere, Category = Settings)
 	TEnumAsByte<EAnimGroupRole::Type> GroupRole = EAnimGroupRole::CanBeLeader;
 
diff --git a/Engine/Source/Runtime/Engine/Classes/Animation/AnimNode_SequencePlayer.h b/Engine/Source/Runtime/Engine/Classes/Animation/AnimNode_SequencePlayer.h
--- a/Engine/Source/Runtime/Engine/Classes/Animation/AnimNode_SequencePlayer.h
+++ b/Engine/Source/Runtime/Engine/Classes/Animation/AnimNode_SequencePlayer.h
@@ -90,15 +90,15 @@ protected:
 	friend class UAnimGraphNode_SequencePlayer;
 
 #if WITH_EDITORONLY_DATA
-	// The group name (NAME_None if it is not part of any group)
+	// The group name that we synchronize with (NAME_None if it is not part of any group). 
 	UPROPERTY(EditAnywhere, Category=Sync, meta=(FoldProperty))
 	FName GroupName = NAME_None;
 
-	// The role this player can assume within the group (ignored if GroupIndex is INDEX_NONE)
+	// The role this node can assume within the group (ignored if GroupName is not set)
 	UPROPERTY(EditAnywhere, Category=Sync, meta=(FoldProperty))
 	TEnumAsByte<EAnimGroupRole::Type> GroupRole = EAnimGroupRole::CanBeLeader;
 
-	// How synchronization is determined
+	// How this node will synchronize with other animations.
 	UPROPERTY(EditAnywhere, Category=Sync, meta=(FoldProperty))
 	EAnimSyncMethod Method = EAnimSyncMethod::DoNotSync;
 
@@ -174,15 +174,15 @@ struct FAnimNode_SequencePlayer_Standalone : public FAnimNode_SequencePlayerBase
 	GENERATED_BODY()
 
 protected:
-	// The group name (NAME_None if it is not part of any group)
+	// The group name that we synchronize with (NAME_None if it is not part of any group). 
 	UPROPERTY(EditAnywhere, Category=Sync)
 	FName GroupName = NAME_None;
 
-	// The role this player can assume within the group (ignored if GroupIndex is INDEX_NONE)
+	// The role this node can assume within the group (ignored if GroupName is not set)
 	UPROPERTY(EditAnywhere, Category=Sync)
 	TEnumAsByte<EAnimGroupRole::Type> GroupRole = EAnimGroupRole::CanBeLeader;
 
-	// How synchronization is determined
+	// How this node will synchronize with other animations.
 	UPROPERTY(EditAnywhere, Category=Sync)
 	EAnimSyncMethod Method = EAnimSyncMethod::DoNotSync;
 
diff --git a/Engine/Source/Runtime/Engine/Classes/Animation/AnimationAsset.h b/Engine/Source/Runtime/Engine/Classes/Animation/AnimationAsset.h
--- a/Engine/Source/Runtime/Engine/Classes/Animation/AnimationAsset.h
+++ b/Engine/Source/Runtime/Engine/Classes/Animation/AnimationAsset.h
@@ -959,15 +959,15 @@ struct FAnimationGroupReference
 {
 	GENERATED_USTRUCT_BODY()
 	
-	// How an asset will synchronize with other assets
+	// How this animation will synchronize with other animations. 
 	UPROPERTY(EditAnywhere, Category=Settings)
 	EAnimSyncMethod Method;
 
-	// The name of the group
+	// The group name that we synchronize with (NAME_None if it is not part of any group). 
 	UPROPERTY(EditAnywhere, Category=Settings, meta = (EditCondition = "Method == EAnimSyncMethod::SyncGroup"))
 	FName GroupName;
 
-	// The type of membership in the group (potential leader, always follower, etc...)
+	// The role this animation can assume within the group (ignored if GroupName is not set)
 	UPROPERTY(EditAnywhere, Category=Settings, meta = (EditCondition = "Method == EAnimSyncMethod::SyncGroup"))
 	TEnumAsByte<EAnimGroupRole::Type> GroupRole;
 
