Fixed mirror node, motion matching node, and blendstack node not supporting sync groups with inertial blends

#rnx minor

diff --git a/Engine/Plugins/Animation/BlendStack/Source/Runtime/Private/AnimNode_BlendStack.cpp b/Engine/Plugins/Animation/BlendStack/Source/Runtime/Private/AnimNode_BlendStack.cpp
--- a/Engine/Plugins/Animation/BlendStack/Source/Runtime/Private/AnimNode_BlendStack.cpp
+++ b/Engine/Plugins/Animation/BlendStack/Source/Runtime/Private/AnimNode_BlendStack.cpp
@@ -10,6 +10,7 @@
 #include "Animation/BlendSpace.h"
 #include "BlendStack/AnimNode_BlendStackInput.h"
 #include "BlendStackAnimEventsFilterScope.h"
+#include "Animation/AnimInertializationSyncScope.h"
 #include "BlendStack/BlendStackDefines.h"
 #include "VisualLogger/VisualLogger.h"
 
@@ -1060,9 +1061,11 @@ void FAnimNode_BlendStack::UpdateAssetPlayer(const FAnimationUpdateContext& Cont
 		BlendTo(Context, AnimationAsset, AnimationTime, bLoop, bMirrored, MirrorDataTable.Get(), BlendTime,
 			RootBoneBlendTime, BlendProfile, BlendOption, bUseInertialBlend, BlendParameters, WantedPlayRate);
 	}
+
+	const bool bDidBlendToRequestAnInertialBlend = bExecuteBlendTo && bUseInertialBlend;
+	UE::Anim::TOptionalScopedGraphMessage<UE::Anim::FAnimInertializationSyncScope> InertializationSync(bDidBlendToRequestAnInertialBlend, Context);
 	
 	UpdatePlayRate(WantedPlayRate);
-
 	Super::UpdateAssetPlayer(Context);
 }
 
diff --git a/Engine/Plugins/Animation/PoseSearch/Source/Runtime/Private/AnimNode_MotionMatching.cpp b/Engine/Plugins/Animation/PoseSearch/Source/Runtime/Private/AnimNode_MotionMatching.cpp
--- a/Engine/Plugins/Animation/PoseSearch/Source/Runtime/Private/AnimNode_MotionMatching.cpp
+++ b/Engine/Plugins/Animation/PoseSearch/Source/Runtime/Private/AnimNode_MotionMatching.cpp
@@ -2,6 +2,7 @@
 
 #include "PoseSearch/AnimNode_MotionMatching.h"
 
+#include "Animation/AnimInertializationSyncScope.h"
 #include "Animation/AnimInstanceProxy.h"
 #include "Animation/AnimNode_Inertialization.h"
 #include "Animation/AnimRootMotionProvider.h"
@@ -212,6 +213,9 @@ void FAnimNode_MotionMatching::UpdateAssetPlayer(const FAnimationUpdateContext&
 			}
 		}
 	}
+
+	const bool bDidBlendToRequestAnInertialBlend = MotionMatchingState.bJumpedToPose && bUseInertialBlend;
+	UE::Anim::TOptionalScopedGraphMessage<UE::Anim::FAnimInertializationSyncScope> InertializationSync(bDidBlendToRequestAnInertialBlend, Context);
 	
 	FAnimNode_BlendStack_Standalone::UpdatePlayRate(MotionMatchingState.WantedPlayRate);
 	FAnimNode_BlendStack_Standalone::UpdateAssetPlayer(Context);
diff --git a/Engine/Source/Runtime/AnimGraphRuntime/Private/AnimNodes/AnimNode_Mirror.cpp b/Engine/Source/Runtime/AnimGraphRuntime/Private/AnimNodes/AnimNode_Mirror.cpp
--- a/Engine/Source/Runtime/AnimGraphRuntime/Private/AnimNodes/AnimNode_Mirror.cpp
+++ b/Engine/Source/Runtime/AnimGraphRuntime/Private/AnimNodes/AnimNode_Mirror.cpp
@@ -5,6 +5,7 @@
 #include "Animation/AnimNode_Inertialization.h"
 #include "Animation/AnimTrace.h"
 #include "AnimationRuntime.h"
+#include "Animation/AnimInertializationSyncScope.h"
 #include "Animation/AttributesRuntime.h"
 #include "Animation/MirrorSyncScope.h"
 
@@ -135,6 +136,8 @@ void FAnimNode_MirrorBase::Update_AnyThread(const FAnimationUpdateContext& Conte
 	GetEvaluateGraphExposedInputs().Execute(Context);
 
 	bool bMirror = GetMirror();
+	bool bRequestedInertialization = false;
+	
 	if (bMirrorStateIsValid && bMirrorState != bMirror)
 	{
 		if (GetBlendTimeOnMirrorStateChange() > SMALL_NUMBER)
@@ -152,6 +155,8 @@ void FAnimNode_MirrorBase::Update_AnyThread(const FAnimationUpdateContext& Conte
 
 				InertializationRequester->RequestInertialization(Request);
 				InertializationRequester->AddDebugRecord(*Context.AnimInstanceProxy, Context.GetCurrentNodeId());
+				
+				bRequestedInertialization = true;
 			}
 			else
 			{
@@ -168,6 +173,7 @@ void FAnimNode_MirrorBase::Update_AnyThread(const FAnimationUpdateContext& Conte
 	}
 	UMirrorDataTable* MirrorDataTable = GetMirrorDataTable();
 	UE::Anim::TOptionalScopedGraphMessage<UE::Anim::FMirrorSyncScope> Message(bMirror, Context, Context, MirrorDataTable);
+	UE::Anim::TOptionalScopedGraphMessage<UE::Anim::FAnimInertializationSyncScope> InertializationSync(bRequestedInertialization, Context);
 
 	bMirrorState = bMirror;
 	bMirrorStateIsValid = true;
