[UE] Fix for AnimDynamics debug draw transforms
- Removed unrequired transform into component space
- Removed unused ToWorldT lambda
- Added a check on bone index validity in the debug draw code to avoid an assert - incorrect debug draw is probably better than an editor crash when a node is not setup correctly.
- Green sphere representing the origin still behaves oddly when in root relative sim space, but I don't think the problem is in the debug draw code.

#rnx

diff --git a/Engine/Source/Runtime/AnimGraphRuntime/Private/BoneControllers/AnimNode_AnimDynamics.cpp b/Engine/Source/Runtime/AnimGraphRuntime/Private/BoneControllers/AnimNode_AnimDynamics.cpp
--- a/Engine/Source/Runtime/AnimGraphRuntime/Private/BoneControllers/AnimNode_AnimDynamics.cpp
+++ b/Engine/Source/Runtime/AnimGraphRuntime/Private/BoneControllers/AnimNode_AnimDynamics.cpp
@@ -73,13 +73,6 @@ void FAnimNode_AnimDynamics::DrawBodies(FComponentSpacePoseContext& InContext, c
 		return;
 	}
 
-	auto ToWorldT = [this](FComponentSpacePoseContext& InPoseContext, const FTransform& SimTransform)
-	{
-		FTransform OutTransform = GetComponentSpaceTransformFromSimSpace(SimulationSpace, InPoseContext, SimTransform);
-		OutTransform *= InPoseContext.AnimInstanceProxy->GetComponentTransform();
-		return OutTransform;
-	};
-
 	auto ToWorldV = [this](FComponentSpacePoseContext& InPoseContext, const FVector& SimLocation)
 	{
 		FVector OutLoc = GetComponentSpaceTransformFromSimSpace(SimulationSpace, InPoseContext, FTransform(SimLocation)).GetTranslation();
@@ -110,7 +103,6 @@ void FAnimNode_AnimDynamics::DrawBodies(FComponentSpacePoseContext& InContext, c
 
 		FTransform Transform(Body.Pose.Orientation, Body.Pose.Position);
 		Transform = GetComponentSpaceTransformFromSimSpace(SimulationSpace, InContext, Transform);
-		Transform *= Proxy->GetComponentTransform();
 
 		Proxy->AnimDrawDebugCoordinateSystem(Transform.GetTranslation(), Transform.Rotator(), 2.0f, false, -1.0f, 0.15f);
 
@@ -141,7 +133,14 @@ void FAnimNode_AnimDynamics::DrawBodies(FComponentSpacePoseContext& InContext, c
 			break;
 			case AnimPhysSimSpaceType::BoneRelative:
 			{
-				Origin = Proxy->GetComponentTransform() * InContext.Pose.GetComponentSpaceTransform(FCompactPoseBoneIndex(RelativeSpaceBone.BoneIndex));
+				Origin = Proxy->GetComponentTransform();
+
+				const FCompactPoseBoneIndex CompactPoseBoneIndex(RelativeSpaceBone.BoneIndex);
+
+				if (InContext.Pose.GetPose().IsValidIndex(CompactPoseBoneIndex)) // Check bone index validity here to avoid a fatal assert in the call to GetComponentSpaceTransform.
+				{
+					Origin *= InContext.Pose.GetComponentSpaceTransform(CompactPoseBoneIndex);
+				}
 			}
 			break;
 			case AnimPhysSimSpaceType::Component:
