[ChaosVD] Fix an issue where solvers were not flagged as server and RBAN solver will have the incorrect ID recorded.

Without this fix, RBAN solver recording is non functional if a Rigid Solver has an id 1 (usual on first recordings of a session), and some visualization flags that color or hides geometry based on if a particles is from server or client does not work.

#tests Tested by doing a recording of out test map for RBAN and playing it back PIE

diff --git a/Engine/Plugins/Experimental/ChaosVD/Source/ChaosVD/Private/ChaosVDScene.cpp b/Engine/Plugins/Experimental/ChaosVD/Source/ChaosVD/Private/ChaosVDScene.cpp
--- a/Engine/Plugins/Experimental/ChaosVD/Source/ChaosVD/Private/ChaosVDScene.cpp
+++ b/Engine/Plugins/Experimental/ChaosVD/Source/ChaosVD/Private/ChaosVDScene.cpp
@@ -258,8 +258,8 @@ void FChaosVDScene::CreateSolverInfoActor(int32 SolverID)
 		const FName FolderPath("ChaosVisualDebugger/SolverDataContainer");
 		SolverDataInfo->SetFolderPath(FolderPath);
 
-		const bool bIsServer = IsSolverForServer(SolverID);
 		FString SolverName = LoadedRecording->GetSolverName_AssumedLocked(SolverID);
+		const bool bIsServer = SolverName.Contains(TEXT("Server"));
 
 		SolverDataInfo->SetSolverID(SolverID);
 		SolverDataInfo->SetSolverName(SolverName);
diff --git a/Engine/Source/Runtime/AnimGraphRuntime/Private/BoneControllers/AnimNode_RigidBody.cpp b/Engine/Source/Runtime/AnimGraphRuntime/Private/BoneControllers/AnimNode_RigidBody.cpp
--- a/Engine/Source/Runtime/AnimGraphRuntime/Private/BoneControllers/AnimNode_RigidBody.cpp
+++ b/Engine/Source/Runtime/AnimGraphRuntime/Private/BoneControllers/AnimNode_RigidBody.cpp
@@ -1068,7 +1068,7 @@ void FAnimNode_RigidBody::InitPhysics(const UAnimInstance* InAnimInstance)
 
 #if WITH_CHAOS_VISUAL_DEBUGGER
 		PhysicsSimulation->GetChaosVDContextData().Id = FChaosVDRuntimeModule::Get().GenerateUniqueID();
-		PhysicsSimulation->GetChaosVDContextData().Id = static_cast<int32>(EChaosVDContextType::Solver);
+		PhysicsSimulation->GetChaosVDContextData().Type = static_cast<int32>(EChaosVDContextType::Solver);
 #endif
 
 		const int32 NumBodies = UsePhysicsAsset->SkeletalBodySetups.Num();
