ModifyCurve anim node now generates debug data

diff --git a/Engine/Source/Runtime/AnimGraphRuntime/Private/AnimNodes/AnimNode_ModifyCurve.cpp b/Engine/Source/Runtime/AnimGraphRuntime/Private/AnimNodes/AnimNode_ModifyCurve.cpp
--- a/Engine/Source/Runtime/AnimGraphRuntime/Private/AnimNodes/AnimNode_ModifyCurve.cpp
+++ b/Engine/Source/Runtime/AnimGraphRuntime/Private/AnimNodes/AnimNode_ModifyCurve.cpp
@@ -173,6 +173,16 @@ void FAnimNode_ModifyCurve::Update_AnyThread(const FAnimationUpdateContext& Cont
 	GetEvaluateGraphExposedInputs().Execute(Context);
 }
 
+void FAnimNode_ModifyCurve::GatherDebugData(FNodeDebugData& DebugData)
+{
+	DECLARE_SCOPE_HIERARCHICAL_COUNTER_ANIMNODE(GatherDebugData)
+	FString DebugLine = DebugData.GetNodeName(this);
+	DebugLine += FString::Printf(TEXT("(Alpha: %.2f)"), Alpha);
+	DebugData.AddDebugItem(DebugLine);
+
+	SourcePose.GatherDebugData(DebugData);
+}
+
 #if WITH_EDITOR
 
 void FAnimNode_ModifyCurve::AddCurve(const FName& InName, float InValue)
diff --git a/Engine/Source/Runtime/AnimGraphRuntime/Public/AnimNodes/AnimNode_ModifyCurve.h b/Engine/Source/Runtime/AnimGraphRuntime/Public/AnimNodes/AnimNode_ModifyCurve.h
--- a/Engine/Source/Runtime/AnimGraphRuntime/Public/AnimNodes/AnimNode_ModifyCurve.h
+++ b/Engine/Source/Runtime/AnimGraphRuntime/Public/AnimNodes/AnimNode_ModifyCurve.h
@@ -73,6 +73,7 @@ struct FAnimNode_ModifyCurve : public FAnimNode_Base
 	ANIMGRAPHRUNTIME_API virtual void Evaluate_AnyThread(FPoseContext& Output) override;
 
 	ANIMGRAPHRUNTIME_API virtual void Update_AnyThread(const FAnimationUpdateContext& Context) override;
+	ANIMGRAPHRUNTIME_API virtual void GatherDebugData(FNodeDebugData& DebugData) override;
 	// End of FAnimNode_Base interface
 
 #if WITH_EDITOR
