[Anim BP Anim Node] AnimNode_CopyPoseFromMesh: Added better display name for CopyCustomAttributes check box

diff --git a/Engine/Source/Runtime/AnimGraphRuntime/Public/AnimNodes/AnimNode_CopyPoseFromMesh.h b/Engine/Source/Runtime/AnimGraphRuntime/Public/AnimNodes/AnimNode_CopyPoseFromMesh.h
--- a/Engine/Source/Runtime/AnimGraphRuntime/Public/AnimNodes/AnimNode_CopyPoseFromMesh.h
+++ b/Engine/Source/Runtime/AnimGraphRuntime/Public/AnimNodes/AnimNode_CopyPoseFromMesh.h
@@ -31,8 +31,8 @@ struct FAnimNode_CopyPoseFromMesh : public FAnimNode_Base
 	UPROPERTY(EditAnywhere, BlueprintReadWrite, Category = Copy, meta = (NeverAsPin))
 	uint8 bCopyCurves : 1;
   
-	/* Copy custom attributes from SouceMeshComponent */
-	UPROPERTY(EditAnywhere, BlueprintReadWrite, Category = Copy, meta = (NeverAsPin))
+	/* Copy custom attributes (animation attributes) from SourceMeshComponent */
+	UPROPERTY(EditAnywhere, BlueprintReadWrite, DisplayName = "Copy Custom Attributes(Animation Attributes)",Category = Copy, meta = (NeverAsPin))
 	bool bCopyCustomAttributes;
 
 	/* Use root space transform to copy to the target pose. By default, it copies their relative transform (bone space)*/
