Fix scaling issue with RBAN
Adding a Cvar to control scaling in simulated space

diff --git a/Engine/Source/Runtime/AnimGraphRuntime/Private/BoneControllers/AnimNode_RigidBody.cpp b/Engine/Source/Runtime/AnimGraphRuntime/Private/BoneControllers/AnimNode_RigidBody.cpp
--- a/Engine/Source/Runtime/AnimGraphRuntime/Private/BoneControllers/AnimNode_RigidBody.cpp
+++ b/Engine/Source/Runtime/AnimGraphRuntime/Private/BoneControllers/AnimNode_RigidBody.cpp
@@ -111,6 +111,10 @@ FAutoConsoleVariableRef CVarRigidBodyNodeDebugDraw(TEXT("p.RigidBodyNode.DebugDr
 bool bRBAN_InitializeBoneReferencesRangeCheckEnabled = true;
 FAutoConsoleVariableRef CVarRigidBodyNodeInitializeBoneReferencesRangeCheckEnabled(TEXT("p.RigidBodyNode.InitializeBoneReferencesRangeCheckEnabled"), bRBAN_InitializeBoneReferencesRangeCheckEnabled, TEXT(""), ECVF_Default);
 
+bool bRBAN_EnableScalingOnSpaceTransform = false;
+FAutoConsoleVariableRef CVarRigidBodyNodeEnableScalingOnSpaceTransform(TEXT("p.RigidBodyNode.EnableScalingOnSpaceTransform"), bRBAN_EnableScalingOnSpaceTransform, TEXT("Enable scaling on space transform for RBAN."));
+
+
 
 // Array of priorities that can be indexed into with CVars, since task priorities cannot be set from scalability .ini
 static UE::Tasks::ETaskPriority GRigidBodyNodeTaskPriorities[] =
@@ -406,7 +410,10 @@ void FAnimNode_RigidBody::CalculateSimulationSpace(
 	// This means we do not support phantom forces resulting from scale changes, but that's ok.
 	// NOTE: If we don't clear the scale, rapid scaling to zero can introduce large phantom forces
 	// leading to major instability in the simulation
-	SpaceTransform.SetScale3D(FVector::One());
+	if (!bRBAN_EnableScalingOnSpaceTransform)
+	{
+		SpaceTransform.SetScale3D(FVector::One());
+	}
 
 	// If the system is disabled, nothing else to do
 	if ((Settings.WorldAlpha == 0.0f) || (Dt < SMALL_NUMBER))
