[AutoRTFM] Enable AutoRTFM in AnimGraphRuntime.

This is reachable from Verse and needs to have transactional support compiled in.

diff --git a/Engine/Source/Runtime/AnimGraphRuntime/AnimGraphRuntime.Build.cs b/Engine/Source/Runtime/AnimGraphRuntime/AnimGraphRuntime.Build.cs
--- a/Engine/Source/Runtime/AnimGraphRuntime/AnimGraphRuntime.Build.cs
+++ b/Engine/Source/Runtime/AnimGraphRuntime/AnimGraphRuntime.Build.cs
@@ -34,7 +34,5 @@ public class AnimGraphRuntime : ModuleRules
 				"GeometryCollectionEngine",
             }
         );
-
-		bDisableAutoRTFMInstrumentation = true;
     }
 }
