Reshuffle layered bone blend node member variables to save 8 bytes

diff --git a/Engine/Source/Runtime/AnimGraphRuntime/Public/AnimNodes/AnimNode_LayeredBoneBlend.h b/Engine/Source/Runtime/AnimGraphRuntime/Public/AnimNodes/AnimNode_LayeredBoneBlend.h
--- a/Engine/Source/Runtime/AnimGraphRuntime/Public/AnimNodes/AnimNode_LayeredBoneBlend.h
+++ b/Engine/Source/Runtime/AnimGraphRuntime/Public/AnimNodes/AnimNode_LayeredBoneBlend.h
@@ -54,39 +54,20 @@
 	UPROPERTY(EditAnywhere, BlueprintReadWrite, editfixedsize, Category=Runtime, meta=(BlueprintCompilerGeneratedDefaults, PinShownByDefault))
 	TArray<float> BlendWeights;
 
-	/** Whether to blend bone rotations in mesh space or in local space */
-	UPROPERTY(EditAnywhere, BlueprintReadWrite, Category=Config)
-	bool bMeshSpaceRotationBlend;
-
-	/** Whether to blend bone scales in mesh space or in local space */
-	UPROPERTY(EditAnywhere, BlueprintReadWrite, Category=Config)
-	bool bMeshSpaceScaleBlend;
-	
-	/** How to blend the layers together */
-	UPROPERTY(EditAnywhere, BlueprintReadWrite, Category=Config)
-	TEnumAsByte<enum ECurveBlendOption::Type>	CurveBlendOption;
-
-	/** Whether to incorporate the per-bone blend weight of the root bone when lending root motion */
-	UPROPERTY(EditAnywhere, Category = Config)
-	bool bBlendRootMotionBasedOnRootBone;
-
-	bool bHasRelevantPoses;
-
-	/*
- 	 * Max LOD that this node is allowed to run
-	 * For example if you have LODThreshold to be 2, it will run until LOD 2 (based on 0 index)
-	 * when the component LOD becomes 3, it will stop update/evaluate
-	 * currently transition would be issue and that has to be re-visited
-	 */
-	UPROPERTY(EditAnywhere, BlueprintReadWrite, Category = Performance, meta = (DisplayName = "LOD Threshold"))
-	int32 LODThreshold;
-
 protected:
+	// transient data to handle weight and target weight
+	// this array changes based on required bones
+	TArray<FPerBoneBlendWeight> DesiredBoneBlendWeights;
+	TArray<FPerBoneBlendWeight> CurrentBoneBlendWeights;
+
 	// Per-bone weights for the skeleton. Serialized as these are only relative to the skeleton, but can potentially
 	// be regenerated at runtime if the GUIDs dont match
 	UPROPERTY()
 	TArray<FPerBoneBlendWeight>	PerBoneBlendWeights;
 
+	// Per-curve source pose index
+	TBaseBlendedCurve<FDefaultAllocator, UE::Anim::FCurveElementIndexed> CurvePoseSourceIndices;
+
 	// Guids for skeleton used to determine whether the PerBoneBlendWeights need rebuilding
 	UPROPERTY()
 	FGuid SkeletonGuid;
@@ -95,27 +76,46 @@
 	UPROPERTY()
 	FGuid VirtualBoneGuid;
 
-	// transient data to handle weight and target weight
-	// this array changes based on required bones
-	TArray<FPerBoneBlendWeight> DesiredBoneBlendWeights;
-	TArray<FPerBoneBlendWeight> CurrentBoneBlendWeights;
-
-	// Per-curve source pose index
-	TBaseBlendedCurve<FDefaultAllocator, UE::Anim::FCurveElementIndexed> CurvePoseSourceIndices;
-
 	// Serial number of the required bones container
 	uint16 RequiredBonesSerialNumber;
+
+public:
+	/*
+ 	 * Max LOD that this node is allowed to run
+	 * For example if you have LODThreshold to be 2, it will run until LOD 2 (based on 0 index)
+	 * when the component LOD becomes 3, it will stop update/evaluate
+	 * currently transition would be issue and that has to be re-visited
+	 */
+	UPROPERTY(EditAnywhere, BlueprintReadWrite, Category = Performance, meta = (DisplayName = "LOD Threshold"))
+	int32 LODThreshold;
+
+	/** Whether to blend bone rotations in mesh space or in local space */
+	UPROPERTY(EditAnywhere, BlueprintReadWrite, Category=Config)
+	bool bMeshSpaceRotationBlend;
+
+	/** Whether to blend bone scales in mesh space or in local space */
+	UPROPERTY(EditAnywhere, BlueprintReadWrite, Category=Config)
+	bool bMeshSpaceScaleBlend;
 	
-public:	
+	/** How to blend the layers together */
+	UPROPERTY(EditAnywhere, BlueprintReadWrite, Category=Config)
+	TEnumAsByte<enum ECurveBlendOption::Type> CurveBlendOption;
+
+	/** Whether to incorporate the per-bone blend weight of the root bone when lending root motion */
+	UPROPERTY(EditAnywhere, Category = Config)
+	bool bBlendRootMotionBasedOnRootBone;
+
+	bool bHasRelevantPoses;
+
 	FAnimNode_LayeredBoneBlend()
 		: BlendMode(ELayeredBoneBlendMode::BranchFilter)
+		, RequiredBonesSerialNumber(0)
+		, LODThreshold(INDEX_NONE)
 		, bMeshSpaceRotationBlend(false)
 		, bMeshSpaceScaleBlend(false)
 		, CurveBlendOption(ECurveBlendOption::Override)
 		, bBlendRootMotionBasedOnRootBone(true)
 		, bHasRelevantPoses(false)
-		, LODThreshold(INDEX_NONE)
-		, RequiredBonesSerialNumber(0)
 	{
 	}
 
