Exposed bShouldStopAllMontages pin in latent Play Montage node

PR #11517: Expose bShouldStopAllMontages pin in latent Play Montage node

diff --git a/Engine/Source/Runtime/AnimGraphRuntime/Private/PlayMontageCallbackProxy.cpp b/Engine/Source/Runtime/AnimGraphRuntime/Private/PlayMontageCallbackProxy.cpp
--- a/Engine/Source/Runtime/AnimGraphRuntime/Private/PlayMontageCallbackProxy.cpp
+++ b/Engine/Source/Runtime/AnimGraphRuntime/Private/PlayMontageCallbackProxy.cpp
@@ -21,11 +21,12 @@ UPlayMontageCallbackProxy* UPlayMontageCallbackProxy::CreateProxyObjectForPlayMo
 	class UAnimMontage* MontageToPlay,
 	float PlayRate,
 	float StartingPosition,
-	FName StartingSection)
+	FName StartingSection,
+	bool bShouldStopAllMontages)
 {
 	UPlayMontageCallbackProxy* Proxy = NewObject<UPlayMontageCallbackProxy>();
 	Proxy->SetFlags(RF_StrongRefOnFrame);
-	Proxy->PlayMontage(InSkeletalMeshComponent, MontageToPlay, PlayRate, StartingPosition, StartingSection);
+	Proxy->PlayMontage(InSkeletalMeshComponent, MontageToPlay, PlayRate, StartingPosition, StartingSection, bShouldStopAllMontages);
 	return Proxy;
 }
 
@@ -34,14 +35,15 @@ bool UPlayMontageCallbackProxy::PlayMontage(class USkeletalMeshComponent* InSkel
 	class UAnimMontage* MontageToPlay, 
 	float PlayRate, 
 	float StartingPosition, 
-	FName StartingSection)
+	FName StartingSection,
+	bool bShouldStopAllMontages)
 {
 	bool bPlayedSuccessfully = false;
 	if (InSkeletalMeshComponent)
 	{
 		if (UAnimInstance* AnimInstance = InSkeletalMeshComponent->GetAnimInstance())
 		{
-			const float MontageLength = AnimInstance->Montage_Play(MontageToPlay, PlayRate, EMontagePlayReturnType::MontageLength, StartingPosition);
+			const float MontageLength = AnimInstance->Montage_Play(MontageToPlay, PlayRate, EMontagePlayReturnType::MontageLength, StartingPosition, bShouldStopAllMontages);
 			bPlayedSuccessfully = (MontageLength > 0.f);
 
 			if (bPlayedSuccessfully)
diff --git a/Engine/Source/Runtime/AnimGraphRuntime/Public/PlayMontageCallbackProxy.h b/Engine/Source/Runtime/AnimGraphRuntime/Public/PlayMontageCallbackProxy.h
--- a/Engine/Source/Runtime/AnimGraphRuntime/Public/PlayMontageCallbackProxy.h
+++ b/Engine/Source/Runtime/AnimGraphRuntime/Public/PlayMontageCallbackProxy.h
@@ -40,7 +40,8 @@ class UPlayMontageCallbackProxy : public UObject
 		class UAnimMontage* MontageToPlay, 
 		float PlayRate = 1.f, 
 		float StartingPosition = 0.f, 
-		FName StartingSection = NAME_None);
+		FName StartingSection = NAME_None,
+		bool bShouldStopAllMontages = true);
 
 public:
 	//~ Begin UObject Interface
@@ -78,5 +79,6 @@ protected:
 		class UAnimMontage* MontageToPlay,
 		float PlayRate = 1.f,
 		float StartingPosition = 0.f,
-		FName StartingSection = NAME_None);
+		FName StartingSection = NAME_None,
+		bool bShouldStopAllMontages = true);
 };
