RBAN  Enable CVar to support scaling correctly

diff --git a/Engine/Source/Runtime/AnimGraphRuntime/Private/BoneControllers/AnimNode_RigidBody.cpp b/Engine/Source/Runtime/AnimGraphRuntime/Private/BoneControllers/AnimNode_RigidBody.cpp
--- a/Engine/Source/Runtime/AnimGraphRuntime/Private/BoneControllers/AnimNode_RigidBody.cpp
+++ b/Engine/Source/Runtime/AnimGraphRuntime/Private/BoneControllers/AnimNode_RigidBody.cpp
@@ -111,7 +111,7 @@ FAutoConsoleVariableRef CVarRigidBodyNodeDebugDraw(TEXT("p.RigidBodyNode.DebugDr
 bool bRBAN_InitializeBoneReferencesRangeCheckEnabled = true;
 FAutoConsoleVariableRef CVarRigidBodyNodeInitializeBoneReferencesRangeCheckEnabled(TEXT("p.RigidBodyNode.InitializeBoneReferencesRangeCheckEnabled"), bRBAN_InitializeBoneReferencesRangeCheckEnabled, TEXT(""), ECVF_Default);
 
-bool bRBAN_EnableScalingOnSpaceTransform = false;
+bool bRBAN_EnableScalingOnSpaceTransform = true;
 FAutoConsoleVariableRef CVarRigidBodyNodeEnableScalingOnSpaceTransform(TEXT("p.RigidBodyNode.EnableScalingOnSpaceTransform"), bRBAN_EnableScalingOnSpaceTransform, TEXT("Enable scaling on space transform for RBAN."));
 
 
