Tracking RBAN memory allocations back to associated Assets and Packages. Updating for MemQuery and Insights.

diff --git a/Engine/Source/Runtime/AnimGraphRuntime/Private/BoneControllers/AnimNode_RigidBody.cpp b/Engine/Source/Runtime/AnimGraphRuntime/Private/BoneControllers/AnimNode_RigidBody.cpp
--- a/Engine/Source/Runtime/AnimGraphRuntime/Private/BoneControllers/AnimNode_RigidBody.cpp
+++ b/Engine/Source/Runtime/AnimGraphRuntime/Private/BoneControllers/AnimNode_RigidBody.cpp
@@ -13,20 +13,22 @@
 #include "Engine/OverlapResult.h"
 #include "Engine/SkeletalMesh.h"
 #include "GameFramework/Pawn.h"
+#include "GameFramework/PawnMovementComponent.h"
 #include "HAL/Event.h"
+#include "HAL/LowLevelMemStats.h"
 #include "HAL/LowLevelMemTracker.h"
-#include "PhysicsEngine/BodySetup.h"
-#include "PhysicsEngine/PhysicsAsset.h"
-#include "PhysicsEngine/PhysicsConstraintTemplate.h"
-#include "PhysicsEngine/SkeletalBodySetup.h"
-#include "GameFramework/PawnMovementComponent.h"
 #include "Physics/Experimental/PhysScene_Chaos.h"
 #include "Physics/ImmediatePhysics/ImmediatePhysicsActorHandle.h"
 #include "Physics/ImmediatePhysics/ImmediatePhysicsSimulation.h"
 #include "Physics/ImmediatePhysics/ImmediatePhysicsStats.h"
 #include "Physics/PhysicsInterfaceCore.h"
+#include "PhysicsEngine/BodySetup.h"
+#include "PhysicsEngine/PhysicsAsset.h"
+#include "PhysicsEngine/PhysicsConstraintTemplate.h"
+#include "PhysicsEngine/SkeletalBodySetup.h"
 #include "PhysicsField/PhysicsFieldComponent.h"
 #include "PhysicsEngine/PhysicsSettings.h"
+#include "ProfilingDebugging/AssetMetadataTrace.h"
 #include "Logging/MessageLog.h"
 #include "Logging/LogMacros.h"
 
@@ -114,7 +116,10 @@ FAutoConsoleVariableRef CVarRigidBodyNodeInitializeBoneReferencesRangeCheckEnabl
 bool bRBAN_EnableScalingOnSpaceTransform = true;
 FAutoConsoleVariableRef CVarRigidBodyNodeEnableScalingOnSpaceTransform(TEXT("p.RigidBodyNode.EnableScalingOnSpaceTransform"), bRBAN_EnableScalingOnSpaceTransform, TEXT("Enable scaling on space transform for RBAN."));
 
-
+#if ENABLE_LOW_LEVEL_MEM_TRACKER
+// This is used for memory tagging purposes
+static FName GRBANClassFName(TEXT("AnimNode_RigidBody"));
+#endif
 
 // Array of priorities that can be indexed into with CVars, since task priorities cannot be set from scalability .ini
 static UE::Tasks::ETaskPriority GRigidBodyNodeTaskPriorities[] =
@@ -502,6 +507,13 @@ void FAnimNode_RigidBody::RunPhysicsSimulation(float DeltaSeconds, const FVector
 	LLM_SCOPE_BYNAME(TEXT("Animation/RigidBody"));
 	SCOPE_CYCLE_COUNTER(STAT_RigidBodyNode_Simulation);
 	CSV_SCOPED_TIMING_STAT(Animation, RigidBodyNodeSimulation);
+
+#if ENABLE_LOW_LEVEL_MEM_TRACKER
+	LLM_SCOPE_DYNAMIC_STAT_OBJECTPATH_FNAME(OwningAssetPackageName, ELLMTagSet::Assets);
+	LLM_SCOPE_DYNAMIC_STAT_OBJECTPATH_FNAME(GRBANClassFName, ELLMTagSet::AssetClasses);
+	UE_TRACE_METADATA_SCOPE_ASSET_FNAME(OwningAssetName, GRBANClassFName, OwningAssetPackageName);
+#endif
+
 	FScopeCycleCounterUObject AdditionalScope(UsePhysicsAsset, GET_STATID(STAT_RigidBodyNode_Simulation));
 
 	const int32 MaxSteps = RBAN_MaxSubSteps;
@@ -535,6 +547,13 @@ void FAnimNode_RigidBody::EvaluateSkeletalControl_AnyThread(FComponentSpacePoseC
 	SCOPE_CYCLE_COUNTER(STAT_RigidBody_Eval);
 	CSV_SCOPED_TIMING_STAT(Animation, RigidBodyEval);
 	SCOPE_CYCLE_COUNTER(STAT_ImmediateEvaluateSkeletalControl);
+
+#if ENABLE_LOW_LEVEL_MEM_TRACKER
+	LLM_SCOPE_DYNAMIC_STAT_OBJECTPATH_FNAME(OwningAssetPackageName, ELLMTagSet::Assets);
+	LLM_SCOPE_DYNAMIC_STAT_OBJECTPATH_FNAME(GRBANClassFName, ELLMTagSet::AssetClasses);
+	UE_TRACE_METADATA_SCOPE_ASSET_FNAME(OwningAssetName, GRBANClassFName, OwningAssetPackageName);
+#endif
+
 	//SCOPED_NAMED_EVENT_TEXT("FAnimNode_RigidBody::EvaluateSkeletalControl_AnyThread", FColor::Magenta);
 
 	if (CVarEnableRigidBodyNodeSimulation.GetValueOnAnyThread() == 0)
@@ -1044,6 +1063,15 @@ void FAnimNode_RigidBody::InitPhysics(const UAnimInstance* InAnimInstance)
 		return;
 	}
 
+#if ENABLE_LOW_LEVEL_MEM_TRACKER
+	OwningAssetPackageName = SkeletalMeshAsset->GetPackage()->GetFName();
+	OwningAssetName = SkeletalMeshComp->GetFName();
+
+	LLM_SCOPE_DYNAMIC_STAT_OBJECTPATH_FNAME(OwningAssetPackageName, ELLMTagSet::Assets);
+	LLM_SCOPE_DYNAMIC_STAT_OBJECTPATH_FNAME(GRBANClassFName, ELLMTagSet::AssetClasses);
+	UE_TRACE_METADATA_SCOPE_ASSET_FNAME(OwningAssetName, GRBANClassFName, OwningAssetPackageName);
+#endif
+
 	const FReferenceSkeleton& SkelMeshRefSkel = SkeletalMeshAsset->GetRefSkeleton();
 	UsePhysicsAsset = GetPhysicsAssetToBeUsed(InAnimInstance);
 
diff --git a/Engine/Source/Runtime/AnimGraphRuntime/Public/BoneControllers/AnimNode_RigidBody.h b/Engine/Source/Runtime/AnimGraphRuntime/Public/BoneControllers/AnimNode_RigidBody.h
--- a/Engine/Source/Runtime/AnimGraphRuntime/Public/BoneControllers/AnimNode_RigidBody.h
+++ b/Engine/Source/Runtime/AnimGraphRuntime/Public/BoneControllers/AnimNode_RigidBody.h
@@ -531,6 +531,11 @@ private:
 	FVector SimSpacePreviousComponentAngularVelocity;
 	FVector SimSpacePreviousBoneLinearVelocity;
 	FVector SimSpacePreviousBoneAngularVelocity;
+
+#if ENABLE_LOW_LEVEL_MEM_TRACKER
+	FName OwningAssetPackageName;
+	FName OwningAssetName;
+#endif
 };
 
 #if WITH_EDITORONLY_DATA
