Fix for aim offset anim node to account for socket offset when calculating pivot location

diff --git a/Engine/Source/Runtime/AnimGraphRuntime/Private/AnimNodes/AnimNode_AimOffsetLookAt.cpp b/Engine/Source/Runtime/AnimGraphRuntime/Private/AnimNodes/AnimNode_AimOffsetLookAt.cpp
--- a/Engine/Source/Runtime/AnimGraphRuntime/Private/AnimNodes/AnimNode_AimOffsetLookAt.cpp
+++ b/Engine/Source/Runtime/AnimGraphRuntime/Private/AnimNodes/AnimNode_AimOffsetLookAt.cpp
@@ -128,7 +128,7 @@ void FAnimNode_AimOffsetLookAt::UpdateFromLookAtTarget(FPoseContext& LocalPoseCo
 		if (PivotSocketBoneReference.IsValidToEvaluate(RequiredBones))
 		{
 			const FCompactPoseBoneIndex PivotSocketBoneIndex = PivotSocketBoneReference.GetCompactPoseIndex(RequiredBones);
-			const FTransform PivotBoneComponentTransform = GlobalPose.GetComponentSpaceTransform(PivotSocketBoneIndex);
+			const FTransform PivotBoneComponentTransform = PivotSocketLocalTransform * GlobalPose.GetComponentSpaceTransform(PivotSocketBoneIndex);
 
 			SourceComponentTransform.SetTranslation(PivotBoneComponentTransform.GetTranslation());
 		}
