Simple hoist of a clamp to outside a loop

diff --git a/Engine/Plugins/Experimental/PhysicsControl/Source/PhysicsControl/Private/AnimNode_RigidBodyWithControl.cpp b/Engine/Plugins/Experimental/PhysicsControl/Source/PhysicsControl/Private/AnimNode_RigidBodyWithControl.cpp
--- a/Engine/Plugins/Experimental/PhysicsControl/Source/PhysicsControl/Private/AnimNode_RigidBodyWithControl.cpp
+++ b/Engine/Plugins/Experimental/PhysicsControl/Source/PhysicsControl/Private/AnimNode_RigidBodyWithControl.cpp
@@ -731,13 +731,19 @@ void FAnimNode_RigidBodyWithControl::EvaluateSkeletalControl_AnyThread(FComponen
 					const FVector ComponentDeltaLocation = CurrentTransform.GetTranslation() - PreviousTransform.GetTranslation();
 					const FVector ComponentLinearVelocity = ComponentDeltaLocation / DeltaSeconds;
 					// Apply acceleration that opposed velocity (basically 'drag')
-					FVector ApplyLinearAcc = WorldVectorToSpaceNoScale(SimulationSpace, -ComponentLinearVelocity, CompWorldSpaceTM, BaseBoneTM) * ComponentLinearVelScale;
+					FVector LinearAccelToApply = WorldVectorToSpaceNoScale(SimulationSpace, -ComponentLinearVelocity, CompWorldSpaceTM, BaseBoneTM) * ComponentLinearVelScale;
 
 					// Calc linear acceleration
 					const FVector ComponentLinearAcceleration = (ComponentLinearVelocity - PreviousComponentLinearVelocity) / DeltaSeconds;
 					PreviousComponentLinearVelocity = ComponentLinearVelocity;
 					// Apply opposite acceleration to bodies
-					ApplyLinearAcc += WorldVectorToSpaceNoScale(SimulationSpace, -ComponentLinearAcceleration, CompWorldSpaceTM, BaseBoneTM) * ComponentLinearAccScale;
+					LinearAccelToApply += WorldVectorToSpaceNoScale(SimulationSpace, -ComponentLinearAcceleration, CompWorldSpaceTM, BaseBoneTM) * ComponentLinearAccScale;
+
+					// Clamp if desired
+					if (!ComponentAppliedLinearAccClamp.IsNearlyZero())
+					{
+						LinearAccelToApply = LinearAccelToApply.BoundToBox(-ComponentAppliedLinearAccClamp, ComponentAppliedLinearAccClamp);
+					}
 
 					// Iterate over bodies
 					for (const UE::PhysicsControl::FOutputBoneData& OutputData : OutputBoneData)
@@ -753,17 +759,8 @@ void FAnimNode_RigidBodyWithControl::EvaluateSkeletalControl_AnyThread(FComponen
 							const float BodyInvMass = (float) Body->GetInverseMass();
 							if (BodyInvMass > 0.f)
 							{
-								// Final desired acceleration to apply to body
-								FVector FinalBodyLinearAcc = ApplyLinearAcc;
-
-								// Clamp if desired
-								if (!ComponentAppliedLinearAccClamp.IsNearlyZero())
-								{
-									FinalBodyLinearAcc = FinalBodyLinearAcc.BoundToBox(-ComponentAppliedLinearAccClamp, ComponentAppliedLinearAccClamp);
-								}
-
 								// Apply to body
-								Body->AddForce(FinalBodyLinearAcc / BodyInvMass);
+								Body->AddForce(LinearAccelToApply / BodyInvMass);
 							}
 						}
 					}
diff --git a/Engine/Source/Runtime/AnimGraphRuntime/Private/BoneControllers/AnimNode_RigidBody.cpp b/Engine/Source/Runtime/AnimGraphRuntime/Private/BoneControllers/AnimNode_RigidBody.cpp
--- a/Engine/Source/Runtime/AnimGraphRuntime/Private/BoneControllers/AnimNode_RigidBody.cpp
+++ b/Engine/Source/Runtime/AnimGraphRuntime/Private/BoneControllers/AnimNode_RigidBody.cpp
@@ -772,13 +772,19 @@ void FAnimNode_RigidBody::EvaluateSkeletalControl_AnyThread(FComponentSpacePoseC
 					const FVector ComponentDeltaLocation = CurrentTransform.GetTranslation() - PreviousTransform.GetTranslation();
 					const FVector ComponentLinearVelocity = ComponentDeltaLocation / DeltaSeconds;
 					// Apply acceleration that opposed velocity (basically 'drag')
-					FVector ApplyLinearAcc = WorldVectorToSpaceNoScale(SimulationSpace, -ComponentLinearVelocity, CompWorldSpaceTM, BaseBoneTM) * UseComponentLinearVelScale;
+					FVector LinearAccelToApply = WorldVectorToSpaceNoScale(SimulationSpace, -ComponentLinearVelocity, CompWorldSpaceTM, BaseBoneTM) * UseComponentLinearVelScale;
 
 					// Calc linear acceleration
 					const FVector ComponentLinearAcceleration = (ComponentLinearVelocity - PreviousComponentLinearVelocity) / DeltaSeconds;
 					PreviousComponentLinearVelocity = ComponentLinearVelocity;
 					// Apply opposite acceleration to bodies
-					ApplyLinearAcc += WorldVectorToSpaceNoScale(SimulationSpace, -ComponentLinearAcceleration, CompWorldSpaceTM, BaseBoneTM) * UseComponentLinearAccScale;
+					LinearAccelToApply += WorldVectorToSpaceNoScale(SimulationSpace, -ComponentLinearAcceleration, CompWorldSpaceTM, BaseBoneTM) * UseComponentLinearAccScale;
+
+					// Clamp if desired
+					if (!UseComponentAppliedLinearAccClamp.IsNearlyZero())
+					{
+						LinearAccelToApply = LinearAccelToApply.BoundToBox(-UseComponentAppliedLinearAccClamp, UseComponentAppliedLinearAccClamp);
+					}
 
 					// Iterate over bodies
 					for (const FOutputBoneData& OutputData : OutputBoneData)
@@ -789,22 +795,11 @@ void FAnimNode_RigidBody::EvaluateSkeletalControl_AnyThread(FComponentSpacePoseC
 						if (BodyData.bIsSimulated)
 						{
 							ImmediatePhysics::FActorHandle* Body = Bodies[BodyIndex];
-
-							// Apply 
 							const double BodyInvMass = Body->GetInverseMass();
 							if (BodyInvMass > 0.0)
 							{
-								// Final desired acceleration to apply to body
-								FVector FinalBodyLinearAcc = ApplyLinearAcc;
-
-								// Clamp if desired
-								if (!UseComponentAppliedLinearAccClamp.IsNearlyZero())
-								{
-									FinalBodyLinearAcc = FinalBodyLinearAcc.BoundToBox(-UseComponentAppliedLinearAccClamp, UseComponentAppliedLinearAccClamp);
-								}
-
 								// Apply to body
-								Body->AddForce(FinalBodyLinearAcc / BodyInvMass);
+								Body->AddForce(LinearAccelToApply / BodyInvMass);
 							}
 						}
 					}
