Allow BlendListBase child classes (BlendListByBool, BlendListByEnum) to always update child nodes.

Since this option doesn't make sense to combine with ResetChildOnActivation, the ResetChildOnActivation boolean has been combined with this option into a new enum: EBlendListChildUpdateMode

diff --git a/Engine/Source/Editor/AnimGraph/Private/AnimGraphNode_BlendListBase.cpp b/Engine/Source/Editor/AnimGraph/Private/AnimGraphNode_BlendListBase.cpp
--- a/Engine/Source/Editor/AnimGraph/Private/AnimGraphNode_BlendListBase.cpp
+++ b/Engine/Source/Editor/AnimGraph/Private/AnimGraphNode_BlendListBase.cpp
@@ -13,6 +13,28 @@ UAnimGraphNode_BlendListBase::UAnimGraphNode_BlendListBase(const FObjectInitiali
 {
 }
 
+void UAnimGraphNode_BlendListBase::Serialize(FArchive& Ar)
+{
+	Super::Serialize(Ar);
+
+	if (Ar.IsLoading())
+	{
+		FAnimNode_BlendListBase* Node = static_cast<FAnimNode_BlendListBase*>(GetFNode());
+
+		if (Node != nullptr)
+		{
+PRAGMA_DISABLE_DEPRECATION_WARNINGS
+			// Handle deprecation
+			if (Node->bResetChildOnActivation_DEPRECATED)
+			{
+				Node->ChildUpateMode = EBlendListChildUpdateMode::ResetChildOnActivate;
+				Node->bResetChildOnActivation_DEPRECATED = false;
+			}
+PRAGMA_ENABLE_DEPRECATION_WARNINGS
+		}
+	}
+}
+
 FLinearColor UAnimGraphNode_BlendListBase::GetNodeTitleColor() const
 {
 	return FLinearColor(0.2f, 0.8f, 0.2f);
diff --git a/Engine/Source/Editor/AnimGraph/Public/AnimGraphNode_BlendListBase.h b/Engine/Source/Editor/AnimGraph/Public/AnimGraphNode_BlendListBase.h
--- a/Engine/Source/Editor/AnimGraph/Public/AnimGraphNode_BlendListBase.h
+++ b/Engine/Source/Editor/AnimGraph/Public/AnimGraphNode_BlendListBase.h
@@ -12,6 +12,8 @@ class ANIMGRAPH_API UAnimGraphNode_BlendListBase : public UAnimGraphNode_Base
 {
 	GENERATED_UCLASS_BODY()
 
+	virtual void Serialize(FArchive& Ar) override;
+
 	// UEdGraphNode interface
 	virtual FLinearColor GetNodeTitleColor() const override;
 	virtual void PostEditChangeProperty(struct FPropertyChangedEvent& PropertyChangedEvent) override;
diff --git a/Engine/Source/Runtime/AnimGraphRuntime/Private/AnimNodes/AnimNode_BlendListBase.cpp b/Engine/Source/Runtime/AnimGraphRuntime/Private/AnimNodes/AnimNode_BlendListBase.cpp
--- a/Engine/Source/Runtime/AnimGraphRuntime/Private/AnimNodes/AnimNode_BlendListBase.cpp
+++ b/Engine/Source/Runtime/AnimGraphRuntime/Private/AnimNodes/AnimNode_BlendListBase.cpp
@@ -195,7 +195,7 @@ void FAnimNode_BlendListBase::Update_AnyThread(const FAnimationUpdateContext& Co
 			}
 
 			// When bResetChildOnActivation is true and the weight of the new child is zero, we'll reinitialize the child.
-			if (GetResetChildOnActivation() && CurrentWeight <= ZERO_ANIMWEIGHT_THRESH)
+			if (GetChildUpdateMode() == EBlendListChildUpdateMode::ResetChildOnActivate && CurrentWeight <= ZERO_ANIMWEIGHT_THRESH)
 			{
 				FAnimationInitializeContext ReinitializeContext(Context.AnimInstanceProxy, Context.SharedContext);
 
@@ -241,6 +241,13 @@ void FAnimNode_BlendListBase::Update_AnyThread(const FAnimationUpdateContext& Co
 				UE::Anim::TOptionalScopedGraphMessage<UE::Anim::FAnimInertializationSyncScope> InertializationSync(bRequestedInertializationOnActiveChildIndexChange, ChildContext);
 				BlendPose[i].Update((i == ChildIndex) ? ChildContext : ChildContext.AsInactive());
 			}
+			else if (GetChildUpdateMode() == EBlendListChildUpdateMode::AlwaysTickChildren)
+			{
+				// With always update children on, even if weight is 0 we will update the child
+				FAnimationUpdateContext ChildContext = Context.FractionalWeight(FAnimWeight::GetSmallestRelevantWeight());
+				UE::Anim::TOptionalScopedGraphMessage<UE::Anim::FAnimInertializationSyncScope> InertializationSync(bRequestedInertializationOnActiveChildIndexChange, ChildContext);
+				BlendPose[i].Update((i == ChildIndex) ? ChildContext : ChildContext.AsInactive());
+			}
 		}
 
 		// If we're using a blend profile, extract the scales and build blend sample data
@@ -384,9 +391,16 @@ EAlphaBlendOption FAnimNode_BlendListBase::GetBlendType() const
 	return GET_ANIM_NODE_DATA(EAlphaBlendOption, BlendType);
 }
 
+PRAGMA_DISABLE_DEPRECATION_WARNINGS
 bool FAnimNode_BlendListBase::GetResetChildOnActivation() const
 {
-	return GET_ANIM_NODE_DATA(bool, bResetChildOnActivation);
+	return GetChildUpdateMode() == EBlendListChildUpdateMode::ResetChildOnActivate;
+}
+PRAGMA_ENABLE_DEPRECATION_WARNINGS
+
+EBlendListChildUpdateMode FAnimNode_BlendListBase::GetChildUpdateMode() const
+{
+	return GET_ANIM_NODE_DATA(EBlendListChildUpdateMode, ChildUpateMode);
 }
 
 UCurveFloat* FAnimNode_BlendListBase::GetCustomBlendCurve() const
diff --git a/Engine/Source/Runtime/AnimGraphRuntime/Public/AnimNodes/AnimNode_BlendListBase.h b/Engine/Source/Runtime/AnimGraphRuntime/Public/AnimNodes/AnimNode_BlendListBase.h
--- a/Engine/Source/Runtime/AnimGraphRuntime/Public/AnimNodes/AnimNode_BlendListBase.h
+++ b/Engine/Source/Runtime/AnimGraphRuntime/Public/AnimNodes/AnimNode_BlendListBase.h
@@ -19,6 +19,19 @@ enum class EBlendListTransitionType : uint8
 	Inertialization
 };
 
+UENUM()
+enum class EBlendListChildUpdateMode : uint8
+{	
+	/** Do not tick inactive children, do not reset on activate */
+	Default,
+
+	/** This reinitializes the re-activated child */
+	ResetChildOnActivate,
+
+	/** Always tick children even if they are not active */
+	AlwaysTickChildren
+};
+
 // Blend list node; has many children
 USTRUCT(BlueprintInternalUseOnly)
 struct FAnimNode_BlendListBase : public FAnimNode_Base
@@ -41,9 +54,12 @@ private:
 	EAlphaBlendOption BlendType = EAlphaBlendOption::Linear;
 	
 protected:
-	/** This reinitializes the re-activated child if the child's weight was zero. */
+	UE_DEPRECATED(5.6, "Use ChildUpateMode instead.")
+	UPROPERTY(meta = (DeprecatedProperty, DeprecationMessage = "Use the ChildUpateMode instead"))
+	bool bResetChildOnActivation_DEPRECATED = false;
+	
 	UPROPERTY(EditAnywhere, Category = Option, meta=(FoldProperty))
-	bool bResetChildOnActivation = false;
+	EBlendListChildUpdateMode ChildUpateMode = EBlendListChildUpdateMode::Default;
 
 private:
 	UPROPERTY(EditAnywhere, Category=BlendType, meta=(FoldProperty))
@@ -71,7 +87,17 @@ protected:
 	int32 LastActiveChildIndex = 0;
 	
 public:	
+PRAGMA_DISABLE_DEPRECATION_WARNINGS
+	// Note: We need to explicitly disable warnings on these constructors/operators for clang to be happy with deprecated variables
+	// this is a requirement for clang to compile without warnings.
 	FAnimNode_BlendListBase() = default;
+	~FAnimNode_BlendListBase() = default;
+	FAnimNode_BlendListBase(const FAnimNode_BlendListBase&) = default;
+	FAnimNode_BlendListBase(FAnimNode_BlendListBase&&) = default;
+	FAnimNode_BlendListBase& operator=(const FAnimNode_BlendListBase&) = default;
+	FAnimNode_BlendListBase& operator=(FAnimNode_BlendListBase&&) = default;
+PRAGMA_ENABLE_DEPRECATION_WARNINGS
+	friend class UAnimGraphNode_BlendListBase;
 
 	// FAnimNode_Base interface
 	ANIMGRAPHRUNTIME_API virtual void Initialize_AnyThread(const FAnimationInitializeContext& Context) override;
@@ -106,8 +132,12 @@ public:
 	ANIMGRAPHRUNTIME_API EAlphaBlendOption GetBlendType() const;
 	
 	/** Get whether to reinitialize the child pose when re-activated. For example, when active child changes */
+	UE_DEPRECATED(5.6, "GetResetChildOnActivation is deprecated, please use GetChildUpdateMode instead.")
 	ANIMGRAPHRUNTIME_API bool GetResetChildOnActivation() const;
 
+	/** Get the child update mode. */
+	ANIMGRAPHRUNTIME_API EBlendListChildUpdateMode GetChildUpdateMode() const;
+
 	// Get the custom blend curve to apply when blending, if any
 	ANIMGRAPHRUNTIME_API UCurveFloat* GetCustomBlendCurve() const;
 
