Allow Blend by Bool to have a seperate BlendProfile option for true and false

diff --git a/Engine/Source/Runtime/AnimGraphRuntime/Private/AnimNodes/AnimNode_BlendListBase.cpp b/Engine/Source/Runtime/AnimGraphRuntime/Private/AnimNodes/AnimNode_BlendListBase.cpp
--- a/Engine/Source/Runtime/AnimGraphRuntime/Private/AnimNodes/AnimNode_BlendListBase.cpp
+++ b/Engine/Source/Runtime/AnimGraphRuntime/Private/AnimNodes/AnimNode_BlendListBase.cpp
@@ -42,7 +42,7 @@ void FAnimNode_BlendListBase::Initialize_AnyThread(const FAnimationInitializeCon
 
 void FAnimNode_BlendListBase::Initialize()
 {
-	UBlendProfile* CurrentBlendProfile = GetBlendProfile();
+	SetCurrentBlendProfile(GetBlendProfile());
 
 	const int32 NumPoses = BlendPose.Num();
 
@@ -63,11 +63,25 @@ void FAnimNode_BlendListBase::Initialize()
 			PerBlendData[i].StartAlpha = 0.0f;
 		}
 	}
+
 	PerBlendData[0].Blend.SetAlpha(1.0f);
+	PerBlendData[0].StartAlpha = 1.0f;
+}
+
+void FAnimNode_BlendListBase::SetCurrentBlendProfile(UBlendProfile* NewBlendProfile)
+{
+	CurrentBlendProfile = NewBlendProfile;
+	if (CurrentBlendProfile && PerBoneSampleData.IsEmpty())
+	{
+		InitializePerBoneData();
+	}
+}
 
+void FAnimNode_BlendListBase::InitializePerBoneData()
+{
 	if (CurrentBlendProfile)
 	{
-		PerBlendData[0].StartAlpha = 1.0f;
+		const int32 NumPoses = BlendPose.Num();
 
 		// Initialise per-bone data
 		PerBoneSampleData.Empty(NumPoses);
@@ -103,13 +117,13 @@ void FAnimNode_BlendListBase::Update_AnyThread(const FAnimationUpdateContext& Co
 	
 	if (NumPoses > 0)
 	{
-		UBlendProfile* CurrentBlendProfile = GetBlendProfile();
-		
 		// Handle a change in the active child index; adjusting the target weights
 		const int32 ChildIndex = GetActiveChildIndex();
 		
 		if (ChildIndex != LastActiveChildIndex)
 		{
+			SetCurrentBlendProfile(GetBlendProfile());	// Blend profile can change based on the active child
+
 			bool LastChildIndexIsInvalid = (LastActiveChildIndex == INDEX_NONE);
 			
 			const float CurrentWeight = PerBlendData[ChildIndex].Weight;
@@ -297,7 +311,6 @@ void FAnimNode_BlendListBase::Evaluate_AnyThread(FPoseContext& Output)
 
 	if ((NumPoses > 0) && (BlendPose.Num() == BlendWeights.Num()))
 	{
-		UBlendProfile* CurrentBlendProfile = GetBlendProfile();
 		if(NumPoses == 1 && FAnimWeight::IsFullWeight(BlendWeights[PosesToEvaluate[0]]) && CurrentBlendProfile == nullptr)
 		{
 			// Single full weight pose - pass-through fast common case
diff --git a/Engine/Source/Runtime/AnimGraphRuntime/Private/AnimNodes/AnimNode_BlendListByBool.cpp b/Engine/Source/Runtime/AnimGraphRuntime/Private/AnimNodes/AnimNode_BlendListByBool.cpp
--- a/Engine/Source/Runtime/AnimGraphRuntime/Private/AnimNodes/AnimNode_BlendListByBool.cpp
+++ b/Engine/Source/Runtime/AnimGraphRuntime/Private/AnimNodes/AnimNode_BlendListByBool.cpp
@@ -17,3 +17,24 @@ bool FAnimNode_BlendListByBool::GetActiveValue() const
 {
 	return GET_ANIM_NODE_DATA(bool, bActiveValue);
 }
+
+bool FAnimNode_BlendListByBool::GetUseSeperateBlendProfiles() const
+{
+	return GET_ANIM_NODE_DATA(bool, bUseSeperateBlendProfileForFalse);
+}
+
+UBlendProfile* FAnimNode_BlendListByBool::GetBlendProfileForFalse() const
+{
+	return GET_ANIM_NODE_DATA(TObjectPtr<UBlendProfile>, BlendProfileForFalse);
+}
+
+UBlendProfile* FAnimNode_BlendListByBool::GetBlendProfile() const
+{
+	// If the option is ticked, we have a seperate profile for the false branch
+	if (GetUseSeperateBlendProfiles() && GetActiveValue() == false)
+	{
+		return GetBlendProfileForFalse();
+	}
+
+	return Super::GetBlendProfile();
+}
diff --git a/Engine/Source/Runtime/AnimGraphRuntime/Public/AnimNodes/AnimNode_BlendListBase.h b/Engine/Source/Runtime/AnimGraphRuntime/Public/AnimNodes/AnimNode_BlendListBase.h
--- a/Engine/Source/Runtime/AnimGraphRuntime/Public/AnimNodes/AnimNode_BlendListBase.h
+++ b/Engine/Source/Runtime/AnimGraphRuntime/Public/AnimNodes/AnimNode_BlendListBase.h
@@ -85,6 +85,10 @@ protected:
 	TArray<FBlendSampleData> PerBoneSampleData;
 
 	int32 LastActiveChildIndex = 0;
+
+	// The blend profile used for the current blend
+	// Note its possible that the blend profile changes based on the active child
+	UBlendProfile* CurrentBlendProfile = nullptr;
 	
 public:	
 PRAGMA_DISABLE_DEPRECATION_WARNINGS
@@ -142,13 +146,16 @@ public:
 	ANIMGRAPHRUNTIME_API UCurveFloat* GetCustomBlendCurve() const;
 
 	// Get the blend profile to use when blending, if any
-	ANIMGRAPHRUNTIME_API UBlendProfile* GetBlendProfile() const;
+	// Note that its possible for the blend profile to change based on the active child
+	ANIMGRAPHRUNTIME_API virtual UBlendProfile* GetBlendProfile() const;
 	
 protected:
 	virtual int32 GetActiveChildIndex() { return 0; }
 	virtual FString GetNodeName(FNodeDebugData& DebugData) { return DebugData.GetNodeName(this); }
 
 	ANIMGRAPHRUNTIME_API void Initialize();
+	void InitializePerBoneData();
+	void SetCurrentBlendProfile(UBlendProfile* NewBlendProfile);	
 
 	friend class UBlendListBaseLibrary;
 };
diff --git a/Engine/Source/Runtime/AnimGraphRuntime/Public/AnimNodes/AnimNode_BlendListByBool.h b/Engine/Source/Runtime/AnimGraphRuntime/Public/AnimNodes/AnimNode_BlendListByBool.h
--- a/Engine/Source/Runtime/AnimGraphRuntime/Public/AnimNodes/AnimNode_BlendListByBool.h
+++ b/Engine/Source/Runtime/AnimGraphRuntime/Public/AnimNodes/AnimNode_BlendListByBool.h
@@ -14,6 +14,18 @@ struct FAnimNode_BlendListByBool : public FAnimNode_BlendListBase
 	GENERATED_BODY()
 private:
 #if WITH_EDITORONLY_DATA
+	// Used in conjunction with bUseSeperateBlendProfileForFalse
+	UPROPERTY(EditAnywhere, Category = BlendType, meta = (UseAsBlendProfile = true, FoldProperty, EditCondition="bUseSeperateBlendProfileForFalse", DisplayAfter="bUseSeperateBlendProfileForFalse"))
+	TObjectPtr<UBlendProfile> BlendProfileForFalse = nullptr;
+
+	/* Specify whether to use a different blend profile for the 'false' branch than the true branch.
+	*
+	* If bUseSeperateBlendProfileForFalse is false (default), then the 'BlendProfile' is used when ActiveValue is both true or false
+	* If bUseSeperateBlendProfileForFalse is true, then the 'BlendProfileForFalse' value is used when the ActiveValue is false, but 'BlendProfile' is used when ActiveValue is true
+	*/ 
+	UPROPERTY(EditAnywhere, Category = BlendType, meta = (FoldProperty, DisplayAfter="BlendProfile"))
+	bool bUseSeperateBlendProfileForFalse = false;
+
 	// Which input should be connected to the output?
 	UPROPERTY(EditAnywhere, Category=Runtime, meta=(PinShownByDefault, FoldProperty))
 	bool bActiveValue = false;
@@ -24,8 +36,14 @@ public:
 
 	// Get which input should be connected to the output
 	ANIMGRAPHRUNTIME_API bool GetActiveValue() const;
+
+	ANIMGRAPHRUNTIME_API bool GetUseSeperateBlendProfiles() const;
+
+	ANIMGRAPHRUNTIME_API UBlendProfile* GetBlendProfileForFalse() const;
 	
 protected:
 	ANIMGRAPHRUNTIME_API virtual int32 GetActiveChildIndex() override;
-	virtual FString GetNodeName(FNodeDebugData& DebugData) override { return DebugData.GetNodeName(this); }
+	virtual FString GetNodeName(FNodeDebugData& DebugData) override { return DebugData.GetNodeName(this); }	
+
+	ANIMGRAPHRUNTIME_API virtual UBlendProfile* GetBlendProfile() const override;
 };
