Sync Group fixes:
- Fixed bOverridePositionWhenJoiningSyncGroupAsLeader not showing in the Sequence/BlendSpace/MotionMatching player UI.
- Allowed CanBeLeader sync role to override sync group position when joining as well.
- Fixed bOverridePositionWhenJoiningSyncGroupAsLeader being defaulted to true in blendspace/sequence standalone players.

#rnx

diff --git a/Engine/Plugins/Animation/PoseSearch/Source/Runtime/Public/PoseSearch/AnimNode_MotionMatching.h b/Engine/Plugins/Animation/PoseSearch/Source/Runtime/Public/PoseSearch/AnimNode_MotionMatching.h
--- a/Engine/Plugins/Animation/PoseSearch/Source/Runtime/Public/PoseSearch/AnimNode_MotionMatching.h
+++ b/Engine/Plugins/Animation/PoseSearch/Source/Runtime/Public/PoseSearch/AnimNode_MotionMatching.h
@@ -84,7 +84,7 @@ private:
 	TEnumAsByte<EAnimGroupRole::Type> GroupRole = EAnimGroupRole::ExclusiveAlwaysLeader;
 	
 	// When enabled, acting as the leader, and using marker-based sync, this asset player will not sync to the previous leader's sync position when joining a sync group and before becoming the leader but instead force everyone else to match its position.
-	UPROPERTY(VisibleAnywhere, Category = Sync, meta = (FoldProperty, EditCondition = "GroupRole == EAnimGroupRole::AlwaysLeader || GroupRole == EAnimGroupRole::ExclusiveAlwaysLeader || GroupRole == EAnimGroupRole::TransitionLeader", EditConditionHides))
+	UPROPERTY(VisibleAnywhere, Category = Sync, meta = (FoldProperty, EditCondition = "GroupRole != EAnimGroupRole::TransitionFollower && GroupRole != EAnimGroupRole::AlwaysFollower", EditConditionHides))
 	bool bOverridePositionWhenJoiningSyncGroupAsLeader = true;
 	
 	// How we should update individual blend space parameters. See dropdown options tooltips.
diff --git a/Engine/Source/Runtime/AnimGraphRuntime/Private/AnimNodes/AnimNode_BlendSpacePlayer.cpp b/Engine/Source/Runtime/AnimGraphRuntime/Private/AnimNodes/AnimNode_BlendSpacePlayer.cpp
--- a/Engine/Source/Runtime/AnimGraphRuntime/Private/AnimNodes/AnimNode_BlendSpacePlayer.cpp
+++ b/Engine/Source/Runtime/AnimGraphRuntime/Private/AnimNodes/AnimNode_BlendSpacePlayer.cpp
@@ -231,6 +231,11 @@ EAnimGroupRole::Type FAnimNode_BlendSpacePlayer::GetGroupRole() const
 	return GET_ANIM_NODE_DATA(TEnumAsByte<EAnimGroupRole::Type>, GroupRole);
 }
 
+bool FAnimNode_BlendSpacePlayer::GetOverridePositionWhenJoiningSyncGroupAsLeader() const
+{
+	return GET_ANIM_NODE_DATA(bool, bOverridePositionWhenJoiningSyncGroupAsLeader);
+}
+
 EAnimSyncMethod FAnimNode_BlendSpacePlayer::GetGroupMethod() const
 {
 	return GET_ANIM_NODE_DATA(EAnimSyncMethod, Method);
@@ -271,6 +276,21 @@ bool FAnimNode_BlendSpacePlayer::SetGroupRole(EAnimGroupRole::Type InRole)
 	return false;
 }
 
+bool FAnimNode_BlendSpacePlayer::SetOverridePositionWhenJoiningSyncGroupAsLeader(bool InOverridePositionWhenJoiningSyncGroupAsLeader)
+{
+#if WITH_EDITORONLY_DATA
+	bOverridePositionWhenJoiningSyncGroupAsLeader = InOverridePositionWhenJoiningSyncGroupAsLeader;
+#endif
+
+	if (bool* OverridePositionPtr = GET_INSTANCE_ANIM_NODE_DATA_PTR(bool, bOverridePositionWhenJoiningSyncGroupAsLeader))
+	{
+		*OverridePositionPtr = InOverridePositionWhenJoiningSyncGroupAsLeader;
+		return true;
+	}
+	
+	return false;
+}
+
 bool FAnimNode_BlendSpacePlayer::SetGroupMethod(EAnimSyncMethod InMethod)
 {
 #if WITH_EDITORONLY_DATA
diff --git a/Engine/Source/Runtime/AnimGraphRuntime/Public/AnimNodes/AnimNode_BlendSpacePlayer.h b/Engine/Source/Runtime/AnimGraphRuntime/Public/AnimNodes/AnimNode_BlendSpacePlayer.h
--- a/Engine/Source/Runtime/AnimGraphRuntime/Public/AnimNodes/AnimNode_BlendSpacePlayer.h
+++ b/Engine/Source/Runtime/AnimGraphRuntime/Public/AnimNodes/AnimNode_BlendSpacePlayer.h
@@ -142,6 +142,10 @@ private:
 	UPROPERTY(EditAnywhere, Category = Sync, meta = (FoldProperty))
 	TEnumAsByte<EAnimGroupRole::Type> GroupRole = EAnimGroupRole::CanBeLeader;
 
+	// When enabled, acting as the leader, and using marker-based sync, this asset player will not sync to the previous leader's sync position when joining a sync group and before becoming the leader but instead force everyone else to match its position.
+	UPROPERTY(EditAnywhere, Category = Sync, meta = (FoldProperty, EditCondition = "GroupRole != EAnimGroupRole::TransitionFollower && GroupRole != EAnimGroupRole::AlwaysFollower", EditConditionHides))
+	bool bOverridePositionWhenJoiningSyncGroupAsLeader = false;
+	
 	// How this node will synchronize with other animations. Note that this determines how the output
 	// of this node is used for synchronization, not of animations contained by it.
 	UPROPERTY(EditAnywhere, Category = Sync, meta = (FoldProperty))
@@ -185,11 +189,13 @@ public:
 	// FAnimNode_AssetPlayerBase interface
 	ANIMGRAPHRUNTIME_API virtual FName GetGroupName() const override;
 	ANIMGRAPHRUNTIME_API virtual EAnimGroupRole::Type GetGroupRole() const override;
+	ANIMGRAPHRUNTIME_API virtual bool GetOverridePositionWhenJoiningSyncGroupAsLeader() const override;
 	ANIMGRAPHRUNTIME_API virtual EAnimSyncMethod GetGroupMethod() const override;
 	ANIMGRAPHRUNTIME_API virtual bool GetIgnoreForRelevancyTest() const override;
 	ANIMGRAPHRUNTIME_API virtual bool IsLooping() const override;
 	ANIMGRAPHRUNTIME_API virtual bool SetGroupName(FName InGroupName) override;
 	ANIMGRAPHRUNTIME_API virtual bool SetGroupRole(EAnimGroupRole::Type InRole) override;
+	ANIMGRAPHRUNTIME_API virtual bool SetOverridePositionWhenJoiningSyncGroupAsLeader(bool InOverridePositionWhenJoiningSyncGroupAsLeader) override;;
 	ANIMGRAPHRUNTIME_API virtual bool SetGroupMethod(EAnimSyncMethod InMethod) override;
 	ANIMGRAPHRUNTIME_API virtual bool SetIgnoreForRelevancyTest(bool bInIgnoreForRelevancyTest) override;
 	// End of FAnimNode_AssetPlayerBase interface
@@ -229,8 +235,8 @@ private:
 	TEnumAsByte<EAnimGroupRole::Type> GroupRole = EAnimGroupRole::CanBeLeader;
 
 	// When enabled, acting as the leader, and using marker-based sync, this asset player will not sync to the previous leader's sync position when joining a sync group and before becoming the leader but instead force everyone else to match its position.
-	UPROPERTY(VisibleAnywhere, Category = Sync, meta = (EditCondition = "GroupRole == EAnimGroupRole::AlwaysLeader || GroupRole == EAnimGroupRole::ExclusiveAlwaysLeader || GroupRole == EAnimGroupRole::TransitionLeader", EditConditionHides))
-	bool bOverridePositionWhenJoiningSyncGroupAsLeader = true;
+	UPROPERTY(EditAnywhere, Category = Sync, meta = (EditCondition = "GroupRole != EAnimGroupRole::TransitionFollower && GroupRole != EAnimGroupRole::AlwaysFollower", EditConditionHides))
+	bool bOverridePositionWhenJoiningSyncGroupAsLeader = false;
 	
 	// How this asset will synchronize with other assets. Note that this determines how the output
 	// of this node is used for synchronization, not of animations contained by it.
diff --git a/Engine/Source/Runtime/Engine/Classes/Animation/AnimNode_SequencePlayer.h b/Engine/Source/Runtime/Engine/Classes/Animation/AnimNode_SequencePlayer.h
--- a/Engine/Source/Runtime/Engine/Classes/Animation/AnimNode_SequencePlayer.h
+++ b/Engine/Source/Runtime/Engine/Classes/Animation/AnimNode_SequencePlayer.h
@@ -98,6 +98,10 @@ protected:
 	UPROPERTY(EditAnywhere, Category=Sync, meta=(FoldProperty))
 	TEnumAsByte<EAnimGroupRole::Type> GroupRole = EAnimGroupRole::CanBeLeader;
 
+	// When enabled, acting as the leader, and using marker-based sync, this asset player will not sync to the previous leader's sync position when joining a sync group and before becoming the leader but instead force everyone else to match its position.
+	UPROPERTY(EditAnywhere, Category = Sync, meta = (FoldProperty, EditCondition = "GroupRole != EAnimGroupRole::TransitionFollower && GroupRole != EAnimGroupRole::AlwaysFollower", EditConditionHides))
+	bool bOverridePositionWhenJoiningSyncGroupAsLeader = false;
+	
 	// How this node will synchronize with other animations.
 	UPROPERTY(EditAnywhere, Category=Sync, meta=(FoldProperty))
 	EAnimSyncMethod Method = EAnimSyncMethod::DoNotSync;
@@ -158,11 +162,13 @@ public:
 	// FAnimNode_AssetPlayerBase interface
 	ENGINE_API virtual FName GetGroupName() const override;
 	ENGINE_API virtual EAnimGroupRole::Type GetGroupRole() const override;
+	ENGINE_API virtual bool GetOverridePositionWhenJoiningSyncGroupAsLeader() const override;
 	ENGINE_API virtual EAnimSyncMethod GetGroupMethod() const override;
 	ENGINE_API virtual bool IsLooping() const override;
 	ENGINE_API virtual bool GetIgnoreForRelevancyTest() const override;
 	ENGINE_API virtual bool SetGroupName(FName InGroupName) override;
 	ENGINE_API virtual bool SetGroupRole(EAnimGroupRole::Type InRole) override;
+	ENGINE_API virtual bool SetOverridePositionWhenJoiningSyncGroupAsLeader(bool InOverridePositionWhenJoiningSyncGroupAsLeader) override;
 	ENGINE_API virtual bool SetGroupMethod(EAnimSyncMethod InMethod) override;
 	ENGINE_API virtual bool SetIgnoreForRelevancyTest(bool bInIgnoreForRelevancyTest) override;
 };
@@ -183,8 +189,8 @@ protected:
 	TEnumAsByte<EAnimGroupRole::Type> GroupRole = EAnimGroupRole::CanBeLeader;
 
 	// When enabled, acting as the leader, and using marker-based sync, this asset player will not sync to the previous leader's sync position when joining a sync group and before becoming the leader but instead force everyone else to match its position.
-	UPROPERTY(VisibleAnywhere, Category = Sync, meta = (EditCondition = "GroupRole == EAnimGroupRole::AlwaysLeader || GroupRole == EAnimGroupRole::ExclusiveAlwaysLeader || GroupRole == EAnimGroupRole::TransitionLeader", EditConditionHides))
-	bool bOverridePositionWhenJoiningSyncGroupAsLeader = true;
+	UPROPERTY(EditAnywhere, Category = Sync, meta = (EditCondition = "GroupRole != EAnimGroupRole::TransitionFollower && GroupRole != EAnimGroupRole::AlwaysFollower", EditConditionHides))
+	bool bOverridePositionWhenJoiningSyncGroupAsLeader = false;
 	
 	// How this node will synchronize with other animations.
 	UPROPERTY(EditAnywhere, Category=Sync)
diff --git a/Engine/Source/Runtime/Engine/Private/Animation/AnimNode_SequencePlayer.cpp b/Engine/Source/Runtime/Engine/Private/Animation/AnimNode_SequencePlayer.cpp
--- a/Engine/Source/Runtime/Engine/Private/Animation/AnimNode_SequencePlayer.cpp
+++ b/Engine/Source/Runtime/Engine/Private/Animation/AnimNode_SequencePlayer.cpp
@@ -245,6 +245,11 @@ EAnimGroupRole::Type FAnimNode_SequencePlayer::GetGroupRole() const
 	return GET_ANIM_NODE_DATA(TEnumAsByte<EAnimGroupRole::Type>, GroupRole);
 }
 
+bool FAnimNode_SequencePlayer::GetOverridePositionWhenJoiningSyncGroupAsLeader() const
+{
+	return GET_ANIM_NODE_DATA(bool, bOverridePositionWhenJoiningSyncGroupAsLeader);
+}
+
 EAnimSyncMethod FAnimNode_SequencePlayer::GetGroupMethod() const
 {
 	return GET_ANIM_NODE_DATA(EAnimSyncMethod, Method);
@@ -285,6 +290,21 @@ bool FAnimNode_SequencePlayer::SetGroupRole(EAnimGroupRole::Type InRole)
 	return false;
 }
 
+bool FAnimNode_SequencePlayer::SetOverridePositionWhenJoiningSyncGroupAsLeader(bool InOverridePositionWhenJoiningSyncGroupAsLeader)
+{
+#if WITH_EDITORONLY_DATA
+	bOverridePositionWhenJoiningSyncGroupAsLeader = InOverridePositionWhenJoiningSyncGroupAsLeader;
+#endif
+
+	if (bool* OverridePtr = GET_INSTANCE_ANIM_NODE_DATA_PTR(bool, bOverridePositionWhenJoiningSyncGroupAsLeader))
+	{
+		*OverridePtr = InOverridePositionWhenJoiningSyncGroupAsLeader;
+		return true;
+	}
+
+	return false;
+}
+
 bool FAnimNode_SequencePlayer::SetGroupMethod(EAnimSyncMethod InMethod)
 {
 #if WITH_EDITORONLY_DATA
