Trivial: fixed typos of word 'default'

diff --git a/Engine/Plugins/FX/Niagara/Source/NiagaraEditor/Private/ViewModels/Stack/NiagaraStackFunctionInput.cpp b/Engine/Plugins/FX/Niagara/Source/NiagaraEditor/Private/ViewModels/Stack/NiagaraStackFunctionInput.cpp
--- a/Engine/Plugins/FX/Niagara/Source/NiagaraEditor/Private/ViewModels/Stack/NiagaraStackFunctionInput.cpp
+++ b/Engine/Plugins/FX/Niagara/Source/NiagaraEditor/Private/ViewModels/Stack/NiagaraStackFunctionInput.cpp
@@ -738,7 +738,7 @@ FText UNiagaraStackFunctionInput::GetValueToolTip() const
 			break;
 		}
 		case EValueMode::UnsupportedDefault:
-			ValueToolTipCache = LOCTEXT("UnsupportedDefault", "The defalut value defined in the script graph\nis custom and can not be shown in the selection stack.");
+			ValueToolTipCache = LOCTEXT("UnsupportedDefault", "The default value defined in the script graph\nis custom and can not be shown in the selection stack.");
 			break;
 		}
 	}
diff --git a/Engine/Plugins/Runtime/Harmonix/Source/HarmonixDspEditor/Private/JsonImporterHelper.h b/Engine/Plugins/Runtime/Harmonix/Source/HarmonixDspEditor/Private/JsonImporterHelper.h
--- a/Engine/Plugins/Runtime/Harmonix/Source/HarmonixDspEditor/Private/JsonImporterHelper.h
+++ b/Engine/Plugins/Runtime/Harmonix/Source/HarmonixDspEditor/Private/JsonImporterHelper.h
@@ -71,7 +71,7 @@ public:
 	static bool TryGetObject(TSharedPtr<FJsonValue> Field, TSharedPtr<FJsonObject>& OutObject);
 
 	// shortcut to reading a number as a bool
-	static bool TryGetBool(TSharedPtr<FJsonValue> Field, bool& OutBool, bool Defalut = false);
+	static bool TryGetBool(TSharedPtr<FJsonValue> Field, bool& OutBool, bool Default = false);
 
 	static bool TryGetString(TSharedPtr<FJsonValue> Field, FString& OutString, const FString& Default = kEmptyString);
 
diff --git a/Engine/Source/Developer/SourceControl/Public/ISourceControlProvider.h b/Engine/Source/Developer/SourceControl/Public/ISourceControlProvider.h
--- a/Engine/Source/Developer/SourceControl/Public/ISourceControlProvider.h
+++ b/Engine/Source/Developer/SourceControl/Public/ISourceControlProvider.h
@@ -186,7 +186,7 @@ public:
 	/** Get the source control status as a series of key value pairs */
 	virtual TMap<EStatus, FString> GetStatus() const = 0;
 
-	/** Quick check if source control is enabled. Specifically, it returns true if a source control provider is set (regardless of whether the provider is available) and false if no provider is set. So all providers except the stub DefalutSourceProvider will return true. */
+	/** Quick check if source control is enabled. Specifically, it returns true if a source control provider is set (regardless of whether the provider is available) and false if no provider is set. So all providers except the stub defaultSourceProvider will return true. */
 	virtual bool IsEnabled() const = 0;
 
 	/**
diff --git a/Engine/Source/Runtime/AnimGraphRuntime/Private/BoneControllers/AnimNode_AnimDynamics.cpp b/Engine/Source/Runtime/AnimGraphRuntime/Private/BoneControllers/AnimNode_AnimDynamics.cpp
--- a/Engine/Source/Runtime/AnimGraphRuntime/Private/BoneControllers/AnimNode_AnimDynamics.cpp
+++ b/Engine/Source/Runtime/AnimGraphRuntime/Private/BoneControllers/AnimNode_AnimDynamics.cpp
@@ -26,7 +26,7 @@ TAutoConsoleVariable<int32> CVarEnableAdaptiveSubstep(TEXT("p.AnimDynamicsAdapti
 TAutoConsoleVariable<int32> CVarAdaptiveSubstepNumDebtFrames(TEXT("p.AnimDynamicsNumDebtFrames"), 5, TEXT("Number of frames to maintain as time debt when using adaptive substepping, this should be at least 1 or the time debt will never be cleared."));
 TAutoConsoleVariable<int32> CVarEnableWind(TEXT("p.AnimDynamicsWind"), 1, TEXT("Enables/Disables anim dynamics wind forces globally."), ECVF_Scalability);
 TAutoConsoleVariable<float> CVarComponentAppliedLinearAccClampOverride(TEXT("p.AnimDynamics.ComponentAppliedLinearAccClampOverride"), -1.0f, TEXT("Override the per asset setting for all axis (X,Y & Z) of ComponentAppliedLinearAccClamp for all Anim Dynamics Nodes. Negative values are ignored."));
-TAutoConsoleVariable<float> CVarGravityScale(TEXT("p.AnimDynamics.GravityScale"), 1.0f, TEXT("Multiplies the defalut gravity and the gravity override on all Anim Dynamics Nodes."));
+TAutoConsoleVariable<float> CVarGravityScale(TEXT("p.AnimDynamics.GravityScale"), 1.0f, TEXT("Multiplies the default gravity and the gravity override on all Anim Dynamics Nodes."));
 
 // FindChainBones
 // 
