Fix RBAN explosion when scale is change at runtime

Fix calcualtions of space velocity etc, in presense of changing and adjusted scale
All RBAN simulation inputs now in simulation space (previously space movement was not)
Simulation space transform is no longer passed into the simulation (except for debug draw)

diff --git a/Engine/Source/Runtime/AnimGraphRuntime/Private/BoneControllers/AnimNode_RigidBody.cpp b/Engine/Source/Runtime/AnimGraphRuntime/Private/BoneControllers/AnimNode_RigidBody.cpp
--- a/Engine/Source/Runtime/AnimGraphRuntime/Private/BoneControllers/AnimNode_RigidBody.cpp
+++ b/Engine/Source/Runtime/AnimGraphRuntime/Private/BoneControllers/AnimNode_RigidBody.cpp
@@ -114,9 +114,6 @@ FAutoConsoleVariableRef CVarRigidBodyNodeDebugDraw(TEXT("p.RigidBodyNode.DebugDr
 bool bRBAN_InitializeBoneReferencesRangeCheckEnabled = true;
 FAutoConsoleVariableRef CVarRigidBodyNodeInitializeBoneReferencesRangeCheckEnabled(TEXT("p.RigidBodyNode.InitializeBoneReferencesRangeCheckEnabled"), bRBAN_InitializeBoneReferencesRangeCheckEnabled, TEXT(""), ECVF_Default);
 
-bool bRBAN_EnableScalingOnSpaceTransform = true;
-FAutoConsoleVariableRef CVarRigidBodyNodeEnableScalingOnSpaceTransform(TEXT("p.RigidBodyNode.EnableScalingOnSpaceTransform"), bRBAN_EnableScalingOnSpaceTransform, TEXT("Enable scaling on space transform for RBAN."));
-
 #if ENABLE_LOW_LEVEL_MEM_TRACKER
 // This is used for memory tagging purposes
 static FName GRBANClassFName(TEXT("AnimNode_RigidBody"));
@@ -239,12 +236,8 @@ FAnimNode_RigidBody::FAnimNode_RigidBody()
 	, CapturedFrozenPose()
 	, CapturedFrozenCurves()
 	, PreviousComponentLinearVelocity(0.0f)
-	, SimSpacePreviousComponentToWorld()
-	, SimSpacePreviousBoneToComponent()
-	, SimSpacePreviousComponentLinearVelocity(0.0f)
-	, SimSpacePreviousComponentAngularVelocity(0.0f)
-	, SimSpacePreviousBoneLinearVelocity(0.0f)
-	, SimSpacePreviousBoneAngularVelocity(0.0f)
+	, PreviousSimulationSpaceTransform()
+	, PreviousPreviousSimulationSpaceTransform()
 {
 }
 
@@ -385,48 +378,34 @@ void FAnimNode_RigidBody::InitSimulationSpace(
 	const FTransform& ComponentToWorld,
 	const FTransform& BoneToComponent)
 {
-	SimSpacePreviousComponentToWorld = ComponentToWorld;
-	SimSpacePreviousBoneToComponent = BoneToComponent;
-	SimSpacePreviousComponentLinearVelocity = FVector::ZeroVector;
-	SimSpacePreviousComponentAngularVelocity = FVector::ZeroVector;
-	SimSpacePreviousBoneLinearVelocity = FVector::ZeroVector;
-	SimSpacePreviousBoneAngularVelocity = FVector::ZeroVector;
+	FTransform SimulationSpaceTransform = SpaceToWorldTransform(SimulationSpace, ComponentToWorld, BoneToComponent);
+	PreviousSimulationSpaceTransform = SimulationSpaceTransform;
+	PreviousPreviousSimulationSpaceTransform = SimulationSpaceTransform;
 }
 
-void FAnimNode_RigidBody::CalculateSimulationSpace(
+void FAnimNode_RigidBody::CalculateSimulationSpaceMotion(
 	ESimulationSpace Space, 
-	const FTransform& ComponentToWorld, 
-	const FTransform& BoneToComponent,
+	const FTransform& SimulationSpaceTransform,
+	const FTransform& ComponentToWorld,
 	const float Dt,
 	const FSimSpaceSettings& Settings,
-	FTransform& SpaceTransform, 
 	FVector& SpaceLinearVel, 
 	FVector& SpaceAngularVel, 
 	FVector& SpaceLinearAcc, 
 	FVector& SpaceAngularAcc)
 {
-	// World-space transform of the simulation space
-	SpaceTransform = SpaceToWorldTransform(Space, ComponentToWorld, BoneToComponent);
 	SpaceLinearVel = FVector::ZeroVector;
 	SpaceAngularVel = FVector::ZeroVector;
 	SpaceLinearAcc = FVector::ZeroVector;
 	SpaceAngularAcc = FVector::ZeroVector;
 
-	// The simulation scale does not change - we scale the inputs and outputs instead.
-	// This means we do not support phantom forces resulting from scale changes, but that's ok.
-	// NOTE: If we don't clear the scale, rapid scaling to zero can introduce large phantom forces
-	// leading to major instability in the simulation
-	if (!bRBAN_EnableScalingOnSpaceTransform)
-	{
-		SpaceTransform.SetScale3D(FVector::One());
-	}
-
 	// If the system is disabled, nothing else to do
-	if ((Settings.WorldAlpha == 0.0f) || (Dt < SMALL_NUMBER))
+	if ((Settings.WorldAlpha == 0.0f) || (Dt < UE_SMALL_NUMBER))
 	{
 		return;
 	}
 
+	// If the simulation is in world space, the simulation space is stationary
 	if (Space == ESimulationSpace::WorldSpace)
 	{
 		SpaceLinearVel = Settings.ExternalLinearVelocity;
@@ -434,67 +413,34 @@ void FAnimNode_RigidBody::CalculateSimulationSpace(
 		return;
 	}
 
-	// World-space component velocity and acceleration
-	FVector CompLinVel = Chaos::FVec3::CalculateVelocity(SimSpacePreviousComponentToWorld.GetTranslation(), ComponentToWorld.GetTranslation(), Dt);
-	FVector CompAngVel = Chaos::FRotation3::CalculateAngularVelocity(SimSpacePreviousComponentToWorld.GetRotation(), ComponentToWorld.GetRotation(), Dt);
-	FVector CompLinAcc = (CompLinVel - SimSpacePreviousComponentLinearVelocity) / Dt;
-	FVector CompAngAcc = (CompAngVel - SimSpacePreviousComponentAngularVelocity) / Dt;
-	SimSpacePreviousComponentToWorld = ComponentToWorld;
-	SimSpacePreviousComponentLinearVelocity = CompLinVel;
-	SimSpacePreviousComponentAngularVelocity = CompAngVel;
-
-	if (Space == ESimulationSpace::ComponentSpace)
-	{
-		CompLinVel.Z *= Settings.VelocityScaleZ;
-		CompLinAcc.Z *= Settings.VelocityScaleZ;
-
-		SpaceLinearVel = CompLinVel.GetClampedToMaxSize(Settings.MaxLinearVelocity) + Settings.ExternalLinearVelocity;
-		SpaceAngularVel = CompAngVel.GetClampedToMaxSize(Settings.MaxAngularVelocity) + Settings.ExternalAngularVelocity;
-		SpaceLinearAcc = CompLinAcc.GetClampedToMaxSize(Settings.MaxLinearAcceleration);
-		SpaceAngularAcc = CompAngAcc.GetClampedToMaxSize(Settings.MaxAngularAcceleration);
-		return;
-	}
-	
-	if (Space == ESimulationSpace::BaseBoneSpace)
-	{
-		// World-space component-relative bone velocity and acceleration
-		FVector BoneLinVel = Chaos::FVec3::CalculateVelocity(SimSpacePreviousBoneToComponent.GetTranslation(), BoneToComponent.GetTranslation(), Dt);
-		FVector BoneAngVel = Chaos::FRotation3::CalculateAngularVelocity(SimSpacePreviousBoneToComponent.GetRotation(), BoneToComponent.GetRotation(), Dt);
-		BoneLinVel = ComponentToWorld.TransformVector(BoneLinVel);
-		BoneAngVel = ComponentToWorld.TransformVector(BoneAngVel);
-		FVector BoneLinAcc = (BoneLinVel - SimSpacePreviousBoneLinearVelocity) / Dt;
-		FVector BoneAngAcc = (BoneAngVel - SimSpacePreviousBoneAngularVelocity) / Dt;
-		SimSpacePreviousBoneToComponent = BoneToComponent;
-		SimSpacePreviousBoneLinearVelocity = BoneLinVel;
-		SimSpacePreviousBoneAngularVelocity = BoneAngVel;
-
-		// World-space bone velocity and acceleration
-		FVector NetAngVel = CompAngVel + BoneAngVel;
-		FVector NetAngAcc = CompAngAcc + BoneAngAcc;
-
-		// If we limit the angular velocity, we also need to limit the component of linear velocity that comes from (angvel x offset)
-		float AngVelScale = 1.0f;
-		float NetAngVelLenSq = static_cast<float>(NetAngVel.SizeSquared());
-		if (NetAngVelLenSq > FMath::Square(Settings.MaxAngularVelocity))
-		{
-			AngVelScale = Settings.MaxAngularVelocity * FMath::InvSqrt(NetAngVelLenSq);
-		}
-
-		// Add the linear velocity and acceleration that comes from rotation of the space about the component
-		// NOTE: Component angular velocity constribution is scaled
-		FVector SpaceCompOffset = ComponentToWorld.TransformVector(BoneToComponent.GetTranslation());
-		FVector NetLinVel = CompLinVel + BoneLinVel + FVector::CrossProduct(AngVelScale * CompAngVel, SpaceCompOffset);
-		FVector NetLinAcc = CompLinAcc + BoneLinAcc + FVector::CrossProduct(AngVelScale * CompAngAcc, SpaceCompOffset);
-
-		NetLinVel.Z *= Settings.VelocityScaleZ;
-		NetLinAcc.Z *= Settings.VelocityScaleZ;
-
-		SpaceLinearVel = NetLinVel.GetClampedToMaxSize(Settings.MaxLinearVelocity) + Settings.ExternalLinearVelocity;
-		SpaceAngularVel = NetAngVel.GetClampedToMaxSize(Settings.MaxAngularVelocity) + Settings.ExternalAngularVelocity;
-		SpaceLinearAcc = NetLinAcc.GetClampedToMaxSize(Settings.MaxLinearAcceleration);
-		SpaceAngularAcc = NetAngAcc.GetClampedToMaxSize(Settings.MaxAngularAcceleration);
-		return;
-	}
+	// World-space component linear velocity and acceleration
+	FVector PrevSpaceLinearVel = Chaos::FVec3::CalculateVelocity(PreviousPreviousSimulationSpaceTransform.GetTranslation(), PreviousSimulationSpaceTransform.GetTranslation(), Dt);
+	SpaceLinearVel = Chaos::FVec3::CalculateVelocity(PreviousSimulationSpaceTransform.GetTranslation(), SimulationSpaceTransform.GetTranslation(), Dt);
+	SpaceLinearAcc = (SpaceLinearVel - PrevSpaceLinearVel) / Dt;
+
+	// World-space component angular velocity and acceleration
+	FVector PrevSpaceAngularVel = Chaos::FRotation3::CalculateAngularVelocity(PreviousPreviousSimulationSpaceTransform.GetRotation(), PreviousSimulationSpaceTransform.GetRotation(), Dt);
+	SpaceAngularVel = Chaos::FRotation3::CalculateAngularVelocity(PreviousSimulationSpaceTransform.GetRotation(), SimulationSpaceTransform.GetRotation(), Dt);
+	SpaceAngularAcc = (SpaceAngularVel - PrevSpaceAngularVel) / Dt;
+
+	// Apply Z scale
+	SpaceLinearVel.Z *= Settings.VelocityScaleZ;
+	SpaceLinearAcc.Z *= Settings.VelocityScaleZ;
+
+	// Clamped world-space motion of the simulation space
+	SpaceLinearVel = SpaceLinearVel.GetClampedToMaxSize(Settings.MaxLinearVelocity) + Settings.ExternalLinearVelocity;
+	SpaceAngularVel = SpaceAngularVel.GetClampedToMaxSize(Settings.MaxAngularVelocity) + Settings.ExternalAngularVelocity;
+	SpaceLinearAcc = SpaceLinearAcc.GetClampedToMaxSize(Settings.MaxLinearAcceleration);
+	SpaceAngularAcc = SpaceAngularAcc.GetClampedToMaxSize(Settings.MaxAngularAcceleration);
+
+	// Transform world-space motion into simulation space
+	SpaceLinearVel = SimulationSpaceTransform.InverseTransformVector(SpaceLinearVel);
+	SpaceAngularVel = SimulationSpaceTransform.InverseTransformVectorNoScale(SpaceAngularVel);
+	SpaceLinearAcc = SimulationSpaceTransform.InverseTransformVector(SpaceLinearAcc);
+	SpaceAngularAcc = SimulationSpaceTransform.InverseTransformVectorNoScale(SpaceAngularAcc);
+
+	PreviousPreviousSimulationSpaceTransform = PreviousSimulationSpaceTransform;
+	PreviousSimulationSpaceTransform = SimulationSpaceTransform;
 }
 
 
@@ -835,32 +781,31 @@ void FAnimNode_RigidBody::EvaluateSkeletalControl_AnyThread(FComponentSpacePoseC
 				if (RBAN_SimSpaceOverride_MaxAngularAcceleration >= 0.0f) UseSimSpaceSettings.MaxAngularAcceleration = RBAN_SimSpaceOverride_MaxAngularAcceleration;
 			}
 
-			FTransform SimulationTransform;
-			FVector SimulationLinearVelocity;
-			FVector SimulationAngularVelocity;
-			FVector SimulationLinearAcceleration;
-			FVector SimulationAngularAcceleration;
-			CalculateSimulationSpace(
+			const FTransform SpaceTransform = SpaceToWorldTransform(SimulationSpace, CompWorldSpaceTM, BaseBoneTM);
+			FVector SimulationSpaceLinearVelocity;
+			FVector SimulationSpaceAngularVelocity;
+			FVector SimulationSpaceLinearAcceleration;
+			FVector SimulationSpaceAngularAcceleration;
+			CalculateSimulationSpaceMotion(
 				SimulationSpace, 
-				CompWorldSpaceTM, 
-				BaseBoneTM,
+				SpaceTransform,
+				CompWorldSpaceTM,
 				DeltaSeconds,
 				UseSimSpaceSettings,
-				SimulationTransform,
-				SimulationLinearVelocity,
-				SimulationAngularVelocity,
-				SimulationLinearAcceleration,
-				SimulationAngularAcceleration);
+				SimulationSpaceLinearVelocity,
+				SimulationSpaceAngularVelocity,
+				SimulationSpaceLinearAcceleration,
+				SimulationSpaceAngularAcceleration);
 
-			UpdateWorldObjects(SimulationTransform);
-			UpdateClothColliderObjects(SimulationTransform);
+			UpdateWorldObjects(SpaceTransform);
+			UpdateClothColliderObjects(SpaceTransform);
 
 			PhysicsSimulation->UpdateSimulationSpace(
-				SimulationTransform, 
-				SimulationLinearVelocity,
-				SimulationAngularVelocity,
-				SimulationLinearAcceleration,
-				SimulationAngularAcceleration);
+				SpaceTransform,
+				SimulationSpaceLinearVelocity,
+				SimulationSpaceAngularVelocity,
+				SimulationSpaceLinearAcceleration,
+				SimulationSpaceAngularAcceleration);
 
 			PhysicsSimulation->SetSimulationSpaceSettings(
 				UseSimSpaceSettings.WorldAlpha,
diff --git a/Engine/Source/Runtime/AnimGraphRuntime/Public/BoneControllers/AnimNode_RigidBody.h b/Engine/Source/Runtime/AnimGraphRuntime/Public/BoneControllers/AnimNode_RigidBody.h
--- a/Engine/Source/Runtime/AnimGraphRuntime/Public/BoneControllers/AnimNode_RigidBody.h
+++ b/Engine/Source/Runtime/AnimGraphRuntime/Public/BoneControllers/AnimNode_RigidBody.h
@@ -348,13 +348,12 @@ private:
 		const FTransform& BoneToComponent);
 
 	// Calculate simulation space transform, velocity etc to pass into the solver
-	ANIMGRAPHRUNTIME_API void CalculateSimulationSpace(
+	ANIMGRAPHRUNTIME_API void CalculateSimulationSpaceMotion(
 		ESimulationSpace Space,
+		const FTransform& SpaceTransform,
 		const FTransform& ComponentToWorld,
-		const FTransform& BoneToComponent,
 		const float Dt,
 		const FSimSpaceSettings& Settings,
-		FTransform& SpaceTransform,
 		FVector& SpaceLinearVel,
 		FVector& SpaceAngularVel,
 		FVector& SpaceLinearAcc,
@@ -525,12 +524,8 @@ private:
 	FVector PreviousComponentLinearVelocity;
 
 	// Used by the world-space to simulation-space motion transfer system in Component- or Bone-Space sims
-	FTransform SimSpacePreviousComponentToWorld;
-	FTransform SimSpacePreviousBoneToComponent;
-	FVector SimSpacePreviousComponentLinearVelocity;
-	FVector SimSpacePreviousComponentAngularVelocity;
-	FVector SimSpacePreviousBoneLinearVelocity;
-	FVector SimSpacePreviousBoneAngularVelocity;
+	FTransform PreviousSimulationSpaceTransform;
+	FTransform PreviousPreviousSimulationSpaceTransform;
 
 #if ENABLE_LOW_LEVEL_MEM_TRACKER
 	FName OwningAssetPackageName;
diff --git a/Engine/Source/Runtime/Engine/Private/PhysicsEngine/ImmediatePhysics/ImmediatePhysicsChaos/ImmediatePhysicsSimulation_Chaos.cpp b/Engine/Source/Runtime/Engine/Private/PhysicsEngine/ImmediatePhysics/ImmediatePhysicsChaos/ImmediatePhysicsSimulation_Chaos.cpp
--- a/Engine/Source/Runtime/Engine/Private/PhysicsEngine/ImmediatePhysics/ImmediatePhysicsChaos/ImmediatePhysicsSimulation_Chaos.cpp
+++ b/Engine/Source/Runtime/Engine/Private/PhysicsEngine/ImmediatePhysics/ImmediatePhysicsChaos/ImmediatePhysicsSimulation_Chaos.cpp
@@ -293,6 +293,7 @@ namespace ImmediatePhysics_Chaos
 			, Evolution(Particles, ParticlePrevXs, ParticlePrevRs, CollisionDetector)
 			, NumActiveDynamicActorHandles(0)
 			, SimulationSpace()
+			, SimulationSpaceTransform(FTransform::Identity)
 			, RollingAverageStepTime(ChaosImmediate_Evolution_InitialStepTime)
 			, NumRollingAverageStepTimes(1)
 			, MaxNumRollingAverageStepTimes(ChaosImmediate_Evolution_DeltaTimeCount)
@@ -344,8 +345,12 @@ namespace ImmediatePhysics_Chaos
 
 		TArray<Chaos::FParticlePair> PotentiallyCollidingPairs;
 
+		// The motion of the simulation space for calculation of phantom forces
 		Chaos::FSimulationSpace SimulationSpace;
 
+		// The simulation space transform, used for debug draw
+		FTransform SimulationSpaceTransform;
+
 		Chaos::FReal FixedStepTime;
 		Chaos::FReal RollingAverageStepTime;
 		int32 NumRollingAverageStepTimes;
@@ -787,7 +792,7 @@ namespace ImmediatePhysics_Chaos
 		const FVector& LinearAcc,
 		const FVector& AngularAcc)
 	{
-		Implementation->SimulationSpace.Transform = Transform;
+		Implementation->SimulationSpaceTransform = Transform;
 		Implementation->SimulationSpace.LinearAcceleration = LinearAcc;
 		Implementation->SimulationSpace.AngularAcceleration = AngularAcc;
 		Implementation->SimulationSpace.LinearVelocity = LinearVel;
@@ -897,7 +902,7 @@ namespace ImmediatePhysics_Chaos
 	{
 		SCOPE_CYCLE_COUNTER(STAT_ImmediateSimulate_Chaos);
 		CVD_SCOPE_TRACE_SOLVER_FRAME(FSimulation, *this);
-		CVD_TRACE_SOLVER_SIMULATION_SPACE(Implementation->SimulationSpace.Transform);
+		CVD_TRACE_SOLVER_SIMULATION_SPACE(Implementation->SimulationSpaceTransform);
 
 		using namespace Chaos;
 
@@ -1135,15 +1140,15 @@ namespace ImmediatePhysics_Chaos
 		{
 			if (bChaosImmediate_DebugDrawParticles)
 			{
-				DebugDraw::DrawParticleTransforms(Implementation->SimulationSpace.Transform, Implementation->Particles.GetActiveStaticParticlesView(), &ChaosImmPhysDebugDebugDrawSettings);
+				DebugDraw::DrawParticleTransforms(Implementation->SimulationSpaceTransform, Implementation->Particles.GetActiveStaticParticlesView(), &ChaosImmPhysDebugDebugDrawSettings);
 			}
 			if (bChaosImmediate_DebugDrawShapes)
 			{
-				DebugDraw::DrawParticleShapes(Implementation->SimulationSpace.Transform, Implementation->Particles.GetActiveStaticParticlesView(), 1.0f, &ChaosImmPhysDebugDebugDrawSettings);
+				DebugDraw::DrawParticleShapes(Implementation->SimulationSpaceTransform, Implementation->Particles.GetActiveStaticParticlesView(), 1.0f, &ChaosImmPhysDebugDebugDrawSettings);
 			}
 			if (bChaosImmediate_DebugDrawBounds)
 			{
-				DebugDraw::DrawParticleBounds(Implementation->SimulationSpace.Transform, Implementation->Particles.GetActiveStaticParticlesView(), 0.0f, &ChaosImmPhysDebugDebugDrawSettings);
+				DebugDraw::DrawParticleBounds(Implementation->SimulationSpaceTransform, Implementation->Particles.GetActiveStaticParticlesView(), 0.0f, &ChaosImmPhysDebugDebugDrawSettings);
 			}
 		}
 #endif
@@ -1157,15 +1162,15 @@ namespace ImmediatePhysics_Chaos
 		{
 			if (bChaosImmediate_DebugDrawParticles)
 			{
-				DebugDraw::DrawParticleTransforms(Implementation->SimulationSpace.Transform, Implementation->Particles.GetActiveKinematicParticlesView(), &ChaosImmPhysDebugDebugDrawSettings);
+				DebugDraw::DrawParticleTransforms(Implementation->SimulationSpaceTransform, Implementation->Particles.GetActiveKinematicParticlesView(), &ChaosImmPhysDebugDebugDrawSettings);
 			}
 			if (bChaosImmediate_DebugDrawShapes)
 			{
-				DebugDraw::DrawParticleShapes(Implementation->SimulationSpace.Transform, Implementation->Particles.GetActiveKinematicParticlesView(), 1.0f, &ChaosImmPhysDebugDebugDrawSettings);
+				DebugDraw::DrawParticleShapes(Implementation->SimulationSpaceTransform, Implementation->Particles.GetActiveKinematicParticlesView(), 1.0f, &ChaosImmPhysDebugDebugDrawSettings);
 			}
 			if (bChaosImmediate_DebugDrawBounds)
 			{
-				DebugDraw::DrawParticleBounds(Implementation->SimulationSpace.Transform, Implementation->Particles.GetActiveKinematicParticlesView(), 0.0f, &ChaosImmPhysDebugDebugDrawSettings);
+				DebugDraw::DrawParticleBounds(Implementation->SimulationSpaceTransform, Implementation->Particles.GetActiveKinematicParticlesView(), 0.0f, &ChaosImmPhysDebugDebugDrawSettings);
 			}
 		}
 #endif
@@ -1179,15 +1184,15 @@ namespace ImmediatePhysics_Chaos
 		{
 			if (bChaosImmediate_DebugDrawParticles)
 			{
-				DebugDraw::DrawParticleTransforms(Implementation->SimulationSpace.Transform, Implementation->Particles.GetActiveParticlesView(), &ChaosImmPhysDebugDebugDrawSettings);
+				DebugDraw::DrawParticleTransforms(Implementation->SimulationSpaceTransform, Implementation->Particles.GetActiveParticlesView(), &ChaosImmPhysDebugDebugDrawSettings);
 			}
 			if (bChaosImmediate_DebugDrawShapes)
 			{
-				DebugDraw::DrawParticleShapes(Implementation->SimulationSpace.Transform, Implementation->Particles.GetActiveParticlesView(), 1.0f, &ChaosImmPhysDebugDebugDrawSettings);
+				DebugDraw::DrawParticleShapes(Implementation->SimulationSpaceTransform, Implementation->Particles.GetActiveParticlesView(), 1.0f, &ChaosImmPhysDebugDebugDrawSettings);
 			}
 			if (bChaosImmediate_DebugDrawBounds)
 			{
-				DebugDraw::DrawParticleBounds(Implementation->SimulationSpace.Transform, Implementation->Particles.GetActiveParticlesView(), 0.0f, &ChaosImmPhysDebugDebugDrawSettings);
+				DebugDraw::DrawParticleBounds(Implementation->SimulationSpaceTransform, Implementation->Particles.GetActiveParticlesView(), 0.0f, &ChaosImmPhysDebugDebugDrawSettings);
 			}
 		}
 #endif
@@ -1201,12 +1206,12 @@ namespace ImmediatePhysics_Chaos
 		{
 			if (bChaosImmediate_DebugDrawCollisions)
 			{
-				DebugDraw::DrawCollisions(Implementation->SimulationSpace.Transform, Implementation->Collisions.GetConstraintAllocator(), 1.0f, &ChaosImmPhysDebugDebugDrawSettings);
+				DebugDraw::DrawCollisions(Implementation->SimulationSpaceTransform, Implementation->Collisions.GetConstraintAllocator(), 1.0f, &ChaosImmPhysDebugDebugDrawSettings);
 			}
 			if (bChaosImmediate_DebugDrawJoints)
 			{
-				DebugDraw::DrawJointConstraints(Implementation->SimulationSpace.Transform, Implementation->Joints.LinearConstraints, 1.0f, ChaosImmediate_DebugDrawJointFeatures, &ChaosImmPhysDebugDebugDrawSettings);
-				DebugDraw::DrawJointConstraints(Implementation->SimulationSpace.Transform, Implementation->Joints.NonLinearConstraints, 1.0f, ChaosImmediate_DebugDrawJointFeatures, &ChaosImmPhysDebugDebugDrawSettings);
+				DebugDraw::DrawJointConstraints(Implementation->SimulationSpaceTransform, Implementation->Joints.LinearConstraints, 1.0f, ChaosImmediate_DebugDrawJointFeatures, &ChaosImmPhysDebugDebugDrawSettings);
+				DebugDraw::DrawJointConstraints(Implementation->SimulationSpaceTransform, Implementation->Joints.NonLinearConstraints, 1.0f, ChaosImmediate_DebugDrawJointFeatures, &ChaosImmPhysDebugDebugDrawSettings);
 			}
 		}
 #endif
@@ -1218,7 +1223,7 @@ namespace ImmediatePhysics_Chaos
 		using namespace Chaos;
 		if (FDebugDrawQueue::IsDebugDrawingEnabled() && bChaosImmediate_DebugDrawSimulationSpace)
 		{
-			DebugDraw::DrawSimulationSpace(Implementation->SimulationSpace, &ChaosImmPhysDebugDebugDrawSettings);
+			DebugDraw::DrawSimulationSpace(Implementation->SimulationSpaceTransform, Implementation->SimulationSpace, &ChaosImmPhysDebugDebugDrawSettings);
 		}
 #endif
 	}
diff --git a/Engine/Source/Runtime/Experimental/Chaos/Private/Chaos/ChaosDebugDraw.cpp b/Engine/Source/Runtime/Experimental/Chaos/Private/Chaos/ChaosDebugDraw.cpp
--- a/Engine/Source/Runtime/Experimental/Chaos/Private/Chaos/ChaosDebugDraw.cpp
+++ b/Engine/Source/Runtime/Experimental/Chaos/Private/Chaos/ChaosDebugDraw.cpp
@@ -1947,19 +1947,19 @@ namespace Chaos
 			}
 		}
 
-		void DrawSimulationSpaceImpl(const FSimulationSpace& SimSpace, const FChaosDebugDrawSettings& Settings)
+		void DrawSimulationSpaceImpl(const FTransform& SpaceTransform, const FSimulationSpace& SimSpace, const FChaosDebugDrawSettings& Settings)
 		{
-			const FVec3 Pos = SimSpace.Transform.GetLocation();
-			const FRotation3& Rot = SimSpace.Transform.GetRotation();
+			const FVec3 Pos = SpaceTransform.GetLocation();
+			const FRotation3& Rot = SpaceTransform.GetRotation();
 			const FMatrix33 Rotm = Rot.ToMatrix();
 			DrawDebugDirectionalArrow(Pos, Pos + Settings.DrawScale * Settings.BodyAxisLen * Rotm.GetAxis(0), Settings.DrawScale * Settings.ArrowSize, FColor::Red, false, UE_KINDA_SMALL_NUMBER, uint8(Settings.DrawPriority), Settings.LineThickness);
 			DrawDebugDirectionalArrow(Pos, Pos + Settings.DrawScale * Settings.BodyAxisLen * Rotm.GetAxis(1), Settings.DrawScale * Settings.ArrowSize, FColor::Green, false, UE_KINDA_SMALL_NUMBER, uint8(Settings.DrawPriority), Settings.LineThickness);
 			DrawDebugDirectionalArrow(Pos, Pos + Settings.DrawScale * Settings.BodyAxisLen * Rotm.GetAxis(2), Settings.DrawScale * Settings.ArrowSize, FColor::Blue, false, UE_KINDA_SMALL_NUMBER, uint8(Settings.DrawPriority), Settings.LineThickness);
 
-			DrawDebugLine(Pos, Pos + Settings.VelScale * SimSpace.LinearVelocity, FColor::Cyan, false, UE_KINDA_SMALL_NUMBER, uint8(Settings.DrawPriority), Settings.LineThickness);
-			DrawDebugLine(Pos, Pos + Settings.AngVelScale * SimSpace.AngularVelocity, FColor::Cyan, false, UE_KINDA_SMALL_NUMBER, uint8(Settings.DrawPriority), Settings.LineThickness);
-			DrawDebugLine(Pos, Pos + 0.01f * Settings.VelScale * SimSpace.LinearAcceleration, FColor::Yellow, false, UE_KINDA_SMALL_NUMBER, uint8(Settings.DrawPriority), Settings.LineThickness);
-			DrawDebugLine(Pos, Pos + 0.01f * Settings.AngVelScale * SimSpace.AngularAcceleration, FColor::Orange, false, UE_KINDA_SMALL_NUMBER, uint8(Settings.DrawPriority), Settings.LineThickness);
+			DrawDebugLine(Pos, Pos + Settings.VelScale * SpaceTransform.TransformVector(SimSpace.LinearVelocity), FColor::Cyan, false, UE_KINDA_SMALL_NUMBER, uint8(Settings.DrawPriority), Settings.LineThickness);
+			DrawDebugLine(Pos, Pos + Settings.AngVelScale * SpaceTransform.TransformVector(SimSpace.AngularVelocity), FColor::Cyan, false, UE_KINDA_SMALL_NUMBER, uint8(Settings.DrawPriority), Settings.LineThickness);
+			DrawDebugLine(Pos, Pos + 0.01f * Settings.VelScale * SpaceTransform.TransformVector(SimSpace.LinearAcceleration), FColor::Yellow, false, UE_KINDA_SMALL_NUMBER, uint8(Settings.DrawPriority), Settings.LineThickness);
+			DrawDebugLine(Pos, Pos + 0.01f * Settings.AngVelScale * SpaceTransform.TransformVector(SimSpace.AngularAcceleration), FColor::Orange, false, UE_KINDA_SMALL_NUMBER, uint8(Settings.DrawPriority), Settings.LineThickness);
 		}
 
 		void DrawConstraintGraphImpl(const FRigidTransform3& SpaceTransform, const Private::FPBDIslandManager& Graph, const FChaosDebugDrawSettings& Settings)
@@ -2566,11 +2566,11 @@ namespace Chaos
 			}
 		}
 
-		void DrawSimulationSpace(const FSimulationSpace& SimSpace, const FChaosDebugDrawSettings* Settings)
+		void DrawSimulationSpace(const FTransform& SpaceTransform, const FSimulationSpace& SimSpace, const FChaosDebugDrawSettings* Settings)
 		{
 			if (FDebugDrawQueue::IsDebugDrawingEnabled())
 			{
-				DrawSimulationSpaceImpl(SimSpace, GetChaosDebugDrawSettings(Settings));
+				DrawSimulationSpaceImpl(SpaceTransform, SimSpace, GetChaosDebugDrawSettings(Settings));
 			}
 		}
 
diff --git a/Engine/Source/Runtime/Experimental/Chaos/Private/Chaos/Evolution/PBDMinEvolution.cpp b/Engine/Source/Runtime/Experimental/Chaos/Private/Chaos/Evolution/PBDMinEvolution.cpp
--- a/Engine/Source/Runtime/Experimental/Chaos/Private/Chaos/Evolution/PBDMinEvolution.cpp
+++ b/Engine/Source/Runtime/Experimental/Chaos/Private/Chaos/Evolution/PBDMinEvolution.cpp
@@ -204,20 +204,6 @@ namespace Chaos
 	void FPBDMinEvolution::Integrate(FReal Dt)
 	{
 		SCOPE_CYCLE_COUNTER(STAT_MinEvolution_Integrate);
-
-		// Simulation space velocity and acceleration
-		FVec3 SpaceV = FVec3(0);	// Velocity
-		FVec3 SpaceW = FVec3(0);	// Angular Velocity
-		FVec3 SpaceA = FVec3(0);	// Acceleration
-		FVec3 SpaceB = FVec3(0);	// Angular Acceleration
-		if (SimulationSpaceSettings.Alpha > 0.0f)
-		{
-			SpaceV = SimulationSpace.Transform.InverseTransformVector(SimulationSpace.LinearVelocity);
-			SpaceW = SimulationSpace.Transform.InverseTransformVector(SimulationSpace.AngularVelocity);
-			SpaceA = SimulationSpace.Transform.InverseTransformVector(SimulationSpace.LinearAcceleration);
-			SpaceB = SimulationSpace.Transform.InverseTransformVector(SimulationSpace.AngularAcceleration);
-		}
-
 		const FVec3 BoundsExpansion = FVec3(CollisionDetector.GetCollisionContainer().GetDetectorSettings().BoundsExpansion);
 
 		for (TTransientPBDRigidParticleHandle<FReal, 3>& Particle : Particles.GetActiveParticlesView())
@@ -247,16 +233,16 @@ namespace Chaos
 				// https://en.wikipedia.org/wiki/Rotating_reference_frame
 				if (SimulationSpaceSettings.Alpha > 0.0f)
 				{
-					const FVec3 CoriolisAcc = SimulationSpaceSettings.CoriolisAlpha * 2.0f * FVec3::CrossProduct(SpaceW, Particle.GetV());
-					const FVec3 CentrifugalAcc = SimulationSpaceSettings.CentrifugalAlpha * FVec3::CrossProduct(SpaceW, FVec3::CrossProduct(SpaceW, XCoM));
-					const FVec3 EulerAcc = SimulationSpaceSettings.EulerAlpha * FVec3::CrossProduct(SpaceB, XCoM);
-					const FVec3 LinearAcc = SimulationSpaceSettings.LinearAccelerationAlpha * SpaceA;
-					const FVec3 AngularAcc = SimulationSpaceSettings.AngularAccelerationAlpha * SpaceB;
-					const FVec3 LinearDragAcc = SimulationSpaceSettings.ExternalLinearEtherDrag * SpaceV;
+					const FVec3 CoriolisAcc = SimulationSpaceSettings.CoriolisAlpha * 2.0f * FVec3::CrossProduct(SimulationSpace.AngularVelocity, Particle.GetV());
+					const FVec3 CentrifugalAcc = SimulationSpaceSettings.CentrifugalAlpha * FVec3::CrossProduct(SimulationSpace.AngularVelocity, FVec3::CrossProduct(SimulationSpace.AngularVelocity, XCoM));
+					const FVec3 EulerAcc = SimulationSpaceSettings.EulerAlpha * FVec3::CrossProduct(SimulationSpace.AngularAcceleration, XCoM);
+					const FVec3 LinearAcc = SimulationSpaceSettings.LinearAccelerationAlpha * SimulationSpace.LinearAcceleration;
+					const FVec3 AngularAcc = SimulationSpaceSettings.AngularAccelerationAlpha * SimulationSpace.AngularAcceleration;
+					const FVec3 LinearDragAcc = SimulationSpaceSettings.ExternalLinearEtherDrag * SimulationSpace.LinearVelocity;
 					DV -= SimulationSpaceSettings.Alpha * (LinearAcc + LinearDragAcc + CoriolisAcc + CentrifugalAcc + EulerAcc) * Dt;
 					DW -= SimulationSpaceSettings.Alpha * AngularAcc * Dt;
-					TargetV = -SimulationSpaceSettings.Alpha * SimulationSpaceSettings.LinearVelocityAlpha * SpaceV;
-					TargetW = -SimulationSpaceSettings.Alpha * SimulationSpaceSettings.AngularVelocityAlpha * SpaceW;
+					TargetV = -SimulationSpaceSettings.Alpha * SimulationSpaceSettings.LinearVelocityAlpha * SimulationSpace.LinearVelocity;
+					TargetW = -SimulationSpaceSettings.Alpha * SimulationSpaceSettings.AngularVelocityAlpha * SimulationSpace.AngularVelocity;
 				}
 
 				// New velocity
diff --git a/Engine/Source/Runtime/Experimental/Chaos/Public/Chaos/ChaosDebugDraw.h b/Engine/Source/Runtime/Experimental/Chaos/Public/Chaos/ChaosDebugDraw.h
--- a/Engine/Source/Runtime/Experimental/Chaos/Public/Chaos/ChaosDebugDraw.h
+++ b/Engine/Source/Runtime/Experimental/Chaos/Public/Chaos/ChaosDebugDraw.h
@@ -225,7 +225,7 @@ namespace Chaos
 		CHAOS_API void DrawJointConstraints(const FRigidTransform3& SpaceTransform, const FPBDJointConstraints& Constraints, FRealSingle ColorScale, const FChaosDebugDrawJointFeatures& FeatureMask = FChaosDebugDrawJointFeatures::MakeDefault(), const FChaosDebugDrawSettings* Settings = nullptr);
 		CHAOS_API void DrawCharacterGroundConstraints(const FRigidTransform3& SpaceTransform, const FCharacterGroundConstraintContainer& Constraints, const FChaosDebugDrawSettings* Settings = nullptr);
 		CHAOS_API void DrawShape(const FRigidTransform3& ShapeTransform, const FImplicitObject* Implicit, const FPerShapeData* Shape, const FColor& Color, const float Duration, const FChaosDebugDrawSettings* Settings = nullptr);
-		CHAOS_API void DrawSimulationSpace(const FSimulationSpace& SimSpace, const FChaosDebugDrawSettings* Settings = nullptr);
+		CHAOS_API void DrawSimulationSpace(const FTransform& SpaceTransform, const FSimulationSpace& SimSpace, const FChaosDebugDrawSettings* Settings = nullptr);
 		CHAOS_API void DrawConstraintGraph(const FRigidTransform3& ShapeTransform, const Private::FPBDIslandManager& Graph, const FChaosDebugDrawSettings* Settings = nullptr);
 		CHAOS_API void DrawConnectionGraph(const FRigidClustering& Clustering, const FChaosDebugDrawSettings* Settings = nullptr);
 		CHAOS_API void DrawCollidingShapes(const FRigidTransform3& SpaceTransform, const FPBDCollisionConstraints& Collisions, FRealSingle ColorScale, const FRealSingle Duration = 0.f, const FChaosDebugDrawSettings* Settings = nullptr);
@@ -247,6 +247,11 @@ namespace Chaos
 		{
 			DrawShape(ShapeTransform, Implicit, Shapes.IsSingleShape() ? Shapes.GetShape() : nullptr, Color, Duration, Settings);
 		}
+
+		UE_DEPRECATED(5.5, "Now requires SpaceTransform as first parameter")
+		inline void DrawSimulationSpace(const FSimulationSpace& SimSpace, const FChaosDebugDrawSettings* Settings = nullptr)
+		{
+		}
 #endif
 	}
 }
diff --git a/Engine/Source/Runtime/Experimental/Chaos/Public/Chaos/Evolution/SimulationSpace.h b/Engine/Source/Runtime/Experimental/Chaos/Public/Chaos/Evolution/SimulationSpace.h
--- a/Engine/Source/Runtime/Experimental/Chaos/Public/Chaos/Evolution/SimulationSpace.h
+++ b/Engine/Source/Runtime/Experimental/Chaos/Public/Chaos/Evolution/SimulationSpace.h
@@ -15,15 +15,14 @@ namespace Chaos
 	{
 	public:
 		FSimulationSpace()
-			: Transform(FRigidTransform3())
-			, LinearAcceleration(0)
+			: LinearAcceleration(0)
 			, AngularAcceleration(0)
 			, LinearVelocity(0)
 			, AngularVelocity(0)
 		{
 		}
 
-		FRigidTransform3 Transform;
+		// All properties should be in simulation space
 		FVec3 LinearAcceleration;
 		FVec3 AngularAcceleration;
 		FVec3 LinearVelocity;
