[fn-anim] Add CVar to force LegIK node to solve every frame
- a.AnimNode.LegIKForceAlwaysSolve
- this can be enabled when legs are jittery

#rnx

diff --git a/Engine/Source/Runtime/AnimGraphRuntime/Private/BoneControllers/AnimNode_LegIK.cpp b/Engine/Source/Runtime/AnimGraphRuntime/Private/BoneControllers/AnimNode_LegIK.cpp
--- a/Engine/Source/Runtime/AnimGraphRuntime/Private/BoneControllers/AnimNode_LegIK.cpp
+++ b/Engine/Source/Runtime/AnimGraphRuntime/Private/BoneControllers/AnimNode_LegIK.cpp
@@ -19,6 +19,7 @@ TAutoConsoleVariable<int32> CVarAnimLegIKEnable(TEXT("a.AnimNode.LegIK.Enable"),
 TAutoConsoleVariable<int32> CVarAnimLegIKMaxIterations(TEXT("a.AnimNode.LegIK.MaxIterations"), 0, TEXT("Leg IK MaxIterations override. 0 = node default, > 0 override."));
 TAutoConsoleVariable<float> CVarAnimLegIKTargetReachStepPercent(TEXT("a.AnimNode.LegIK.TargetReachStepPercent"), 0.7f, TEXT("Leg IK TargetReachStepPercent."));
 TAutoConsoleVariable<float> CVarAnimLegIKPullDistribution(TEXT("a.AnimNode.LegIK.PullDistribution"), 0.5f, TEXT("Leg IK PullDistribution. 0 = foot, 0.5 = balanced, 1.f = hip"));
+TAutoConsoleVariable<int32> CVarAnimLegIKForceAlwaysSolve(TEXT("a.AnimNode.LegIK.ForceAlwaysSolve"), 0, TEXT("Leg IK Always Run IK Solver. 0 = default behavior, 1 = Run IK Solver every frame."));
 
 /////////////////////////////////////////////////////
 // FAnimAnimNode_LegIK
@@ -417,7 +418,7 @@ bool FAnimNode_LegIK::DoLegReachIK(FAnimLegIKData& InLegData)
 	const bool bUsingSoftIK = SoftPercentLength < 1.0f && SoftAlpha > 0.f;
 	const bool bFootAtGoal = FootFKLocation.Equals(FootIKLocation, ReachPrecision);
 	const bool bUsingRotationLimit = InLegData.LegDefPtr->bEnableRotationLimit;
-	const bool bNeedsSolver = !bFootAtGoal || bUsingRotationLimit || bUsingSoftIK;
+	const bool bNeedsSolver = !bFootAtGoal || bUsingRotationLimit || bUsingSoftIK || (CVarAnimLegIKForceAlwaysSolve.GetValueOnAnyThread() == 1);
 	if (!bNeedsSolver && !bHasTwistOffset)
 	{
 		return false;
