[UE] Fix for P0 client crash on invalid bone index returned from PhysicsAsset Initialization
- Prevent creation of cloth collider book keeping objects in RBAN when an invalid bone index is returned from PhysicsAssetInitialization

diff --git a/Engine/Source/Runtime/AnimGraphRuntime/Private/BoneControllers/AnimNode_RigidBody.cpp b/Engine/Source/Runtime/AnimGraphRuntime/Private/BoneControllers/AnimNode_RigidBody.cpp
--- a/Engine/Source/Runtime/AnimGraphRuntime/Private/BoneControllers/AnimNode_RigidBody.cpp
+++ b/Engine/Source/Runtime/AnimGraphRuntime/Private/BoneControllers/AnimNode_RigidBody.cpp
@@ -1615,9 +1615,12 @@ void FAnimNode_RigidBody::CollectClothColliderObjects(const USkeletalMeshCompone
 				{
 					FBodyInstance* const BodyInstance = BodyInstances[BodyInstanceIndex];
 
-					ImmediatePhysics::FActorHandle* const ActorHandle = PhysicsSimulation->CreateActor(ImmediatePhysics::EActorType::KinematicActor, BodyInstance, BodyInstance->GetUnrealWorldTransform());
-					PhysicsSimulation->AddToCollidingPairs(ActorHandle); // <-allow collision between this actor and all dynamic actors.
-					ClothColliders.Add(FClothCollider(ActorHandle, SourceComponent, BodyInstance->InstanceBoneIndex));
+					if (BodyInstance->InstanceBoneIndex != INDEX_NONE) // Invalid Bone index can occur if the physics asset references a bone name that does not exist in the Cloth Collision Source.
+					{
+						ImmediatePhysics::FActorHandle* const ActorHandle = PhysicsSimulation->CreateActor(ImmediatePhysics::EActorType::KinematicActor, BodyInstance, BodyInstance->GetUnrealWorldTransform());
+						PhysicsSimulation->AddToCollidingPairs(ActorHandle); // <-allow collision between this actor and all dynamic actors.
+						ClothColliders.Add(FClothCollider(ActorHandle, SourceComponent, BodyInstance->InstanceBoneIndex));
+					}
 
 					// Terminate the instance.
 					if (BodyInstance->IsValidBodyInstance())
