Chaos
- Add world-space damping alpha to RBAN Simulation Space settings
- Fix drag force from ExternalLinearDamping setting to include body velocity (was just space velocity before)
- Fix phantom contrifugal force when WorldAlpha is less than 1.0. CF contains Angvel-squared but the alpha was not squared. Net result was that CF was not scaled down as much as it should have been

diff --git a/Engine/Plugins/Experimental/PhysicsControl/Source/ControlRigPhysics/Private/RigPhysicsSolver.cpp b/Engine/Plugins/Experimental/PhysicsControl/Source/ControlRigPhysics/Private/RigPhysicsSolver.cpp
--- a/Engine/Plugins/Experimental/PhysicsControl/Source/ControlRigPhysics/Private/RigPhysicsSolver.cpp
+++ b/Engine/Plugins/Experimental/PhysicsControl/Source/ControlRigPhysics/Private/RigPhysicsSolver.cpp
@@ -892,15 +892,15 @@ void FRigPhysicsSolver::StepSimulation(
 
 	// Simulation space
 	Chaos::FSimulationSpaceSettings ChaosSimulationSpaceSettings;
-	ChaosSimulationSpaceSettings.Alpha = SimulationSpaceSettings.SpaceMovementAmount;
+	ChaosSimulationSpaceSettings.bEnabled = (SimulationSpaceSettings.SpaceMovementAmount > 0.0f);
 	ChaosSimulationSpaceSettings.ExternalLinearEtherDrag = SimulationSpaceSettings.ExternalLinearDragV;
 	Simulation->SetSimulationSpaceSettings(ChaosSimulationSpaceSettings);
 	Simulation->UpdateSimulationSpace(
 		SimulationSpaceState.SimulationSpaceTM,
-		SimulationSpaceData.LinearVelocity,
-		SimulationSpaceData.AngularVelocity,
-		SimulationSpaceData.LinearAcceleration,
-		SimulationSpaceData.AngularAcceleration);
+		SimulationSpaceSettings.SpaceMovementAmount * SimulationSpaceData.LinearVelocity,
+		SimulationSpaceSettings.SpaceMovementAmount * SimulationSpaceData.AngularVelocity,
+		SimulationSpaceSettings.SpaceMovementAmount * SimulationSpaceData.LinearAcceleration,
+		SimulationSpaceSettings.SpaceMovementAmount * SimulationSpaceData.AngularAcceleration);
 
 	UpdatePrePhysics(ExecuteContext, PhysicsDeltaTime, SimulationSpaceData);
 
diff --git a/Engine/Plugins/Experimental/PhysicsControl/Source/PhysicsControl/Private/AnimNode_RigidBodyWithControl.cpp b/Engine/Plugins/Experimental/PhysicsControl/Source/PhysicsControl/Private/AnimNode_RigidBodyWithControl.cpp
--- a/Engine/Plugins/Experimental/PhysicsControl/Source/PhysicsControl/Private/AnimNode_RigidBodyWithControl.cpp
+++ b/Engine/Plugins/Experimental/PhysicsControl/Source/PhysicsControl/Private/AnimNode_RigidBodyWithControl.cpp
@@ -389,6 +389,12 @@ void FAnimNode_RigidBodyWithControl::CalculateSimulationSpaceMotion(
 	SpaceLinearAcc = SimulationSpaceTransform.InverseTransformVector(SpaceLinearAcc);
 	SpaceAngularAcc = SimulationSpaceTransform.InverseTransformVectorNoScale(SpaceAngularAcc);
 
+	// Apply WorldAlpha to simulation space motion (usually to reduce it)
+	SpaceLinearVel *= Settings.WorldAlpha;
+	SpaceLinearAcc *= Settings.WorldAlpha;
+	SpaceAngularVel *= Settings.WorldAlpha;
+	SpaceAngularAcc *= Settings.WorldAlpha;
+
 	PreviousPreviousSimulationSpaceTransform = PreviousSimulationSpaceTransform;
 	PreviousSimulationSpaceTransform = SimulationSpaceTransform;
 }
@@ -799,7 +805,7 @@ void FAnimNode_RigidBodyWithControl::EvaluateSkeletalControl_AnyThread(FComponen
 				SimulationSpaceLinearAcceleration, SimulationSpaceAngularAcceleration);
 
 			PhysicsSimulation->SetSimulationSpaceSettings(
-				SimSpaceSettingsToUse->WorldAlpha, SimSpaceSettingsToUse->ExternalLinearDragV);
+				(SimSpaceSettingsToUse->WorldAlpha > 0.0f), SimSpaceSettingsToUse->DampingAlpha, SimSpaceSettingsToUse->ExternalLinearDragV);
 
 			PhysicsSimulation->SetSolverSettings(
 				SolverSettings.FixedTimeStep,
diff --git a/Engine/Source/Runtime/AnimGraphRuntime/Private/BoneControllers/AnimNode_RigidBody.cpp b/Engine/Source/Runtime/AnimGraphRuntime/Private/BoneControllers/AnimNode_RigidBody.cpp
--- a/Engine/Source/Runtime/AnimGraphRuntime/Private/BoneControllers/AnimNode_RigidBody.cpp
+++ b/Engine/Source/Runtime/AnimGraphRuntime/Private/BoneControllers/AnimNode_RigidBody.cpp
@@ -77,6 +77,7 @@ FAutoConsoleVariableRef CVarRigidBodyNodeIncludeClothColliders(TEXT("p.RigidBody
 bool bRBAN_SimSpace_EnableOverride = false;
 float RBAN_SimSpaceOverride_WorldAlpha = -1.0f;
 float RBAN_SimSpaceOverride_VelocityScaleZ = -1.0f;
+float RBAN_SimSpaceOverride_DampingAlpha = -1.0f;
 float RBAN_SimSpaceOverride_MaxLinearVelocity = -1.0f;
 float RBAN_SimSpaceOverride_MaxAngularVelocity = -1.0f;
 float RBAN_SimSpaceOverride_MaxLinearAcceleration = -1.0f;
@@ -84,6 +85,7 @@ float RBAN_SimSpaceOverride_MaxAngularAcceleration = -1.0f;
 FAutoConsoleVariableRef CVarRigidBodyNodeSpaceOverride(TEXT("p.RigidBodyNode.Space.Override"), bRBAN_SimSpace_EnableOverride, TEXT("Force-enable the advanced simulation space movement forces"), ECVF_Default);
 FAutoConsoleVariableRef CVarRigidBodyNodeSpaceWorldAlpha(TEXT("p.RigidBodyNode.Space.WorldAlpha"), RBAN_SimSpaceOverride_WorldAlpha, TEXT("RBAN SimSpaceSettings overrides"), ECVF_Default);
 FAutoConsoleVariableRef CVarRigidBodyNodeSpaceVelScaleZ(TEXT("p.RigidBodyNode.Space.VelocityScaleZ"), RBAN_SimSpaceOverride_VelocityScaleZ, TEXT("RBAN SimSpaceSettings overrides"), ECVF_Default);
+FAutoConsoleVariableRef CVarRigidBodyNodeSpaceDampingAlpha(TEXT("p.RigidBodyNode.Space.DampingAlpha"), RBAN_SimSpaceOverride_DampingAlpha, TEXT("RBAN SimSpaceSettings overrides"), ECVF_Default);
 FAutoConsoleVariableRef CVarRigidBodyNodeSpaceMaxCompLinVel(TEXT("p.RigidBodyNode.Space.MaxLinearVelocity"), RBAN_SimSpaceOverride_MaxLinearVelocity, TEXT("RBAN SimSpaceSettings overrides"), ECVF_Default);
 FAutoConsoleVariableRef CVarRigidBodyNodeSpaceMaxCompAngVel(TEXT("p.RigidBodyNode.Space.MaxAngularVelocity"), RBAN_SimSpaceOverride_MaxAngularVelocity, TEXT("RBAN SimSpaceSettings overrides"), ECVF_Default);
 FAutoConsoleVariableRef CVarRigidBodyNodeSpaceMaxCompLinAcc(TEXT("p.RigidBodyNode.Space.MaxLinearAcceleration"), RBAN_SimSpaceOverride_MaxLinearAcceleration, TEXT("RBAN SimSpaceSettings overrides"), ECVF_Default);
@@ -145,6 +147,7 @@ static_assert(std::is_trivially_copyable<FSimSpaceSettings>::value);
 FSimSpaceSettings::FSimSpaceSettings()
 	: WorldAlpha(0)
 	, VelocityScaleZ(1)
+	, DampingAlpha(1)
 	, MaxLinearVelocity(10000)
 	, MaxAngularVelocity(10000)
 	, MaxLinearAcceleration(10000)
@@ -439,6 +442,12 @@ void FAnimNode_RigidBody::CalculateSimulationSpaceMotion(
 	SpaceLinearAcc = SimulationSpaceTransform.InverseTransformVector(SpaceLinearAcc);
 	SpaceAngularAcc = SimulationSpaceTransform.InverseTransformVectorNoScale(SpaceAngularAcc);
 
+	// Apply WorldAlpha to simulation space motion (usually to reduce it)
+	SpaceLinearVel *= Settings.WorldAlpha;
+	SpaceLinearAcc *= Settings.WorldAlpha;
+	SpaceAngularVel *= Settings.WorldAlpha;
+	SpaceAngularAcc *= Settings.WorldAlpha;
+
 	PreviousPreviousSimulationSpaceTransform = PreviousSimulationSpaceTransform;
 	PreviousSimulationSpaceTransform = SimulationSpaceTransform;
 }
@@ -775,6 +784,7 @@ void FAnimNode_RigidBody::EvaluateSkeletalControl_AnyThread(FComponentSpacePoseC
 			{
 				if (RBAN_SimSpaceOverride_WorldAlpha >= 0.0f) UseSimSpaceSettings.WorldAlpha = RBAN_SimSpaceOverride_WorldAlpha;
 				if (RBAN_SimSpaceOverride_VelocityScaleZ >= 0.0f) UseSimSpaceSettings.VelocityScaleZ = RBAN_SimSpaceOverride_VelocityScaleZ;
+				if (RBAN_SimSpaceOverride_DampingAlpha >= 0.0f) UseSimSpaceSettings.DampingAlpha = RBAN_SimSpaceOverride_DampingAlpha;
 				if (RBAN_SimSpaceOverride_MaxLinearVelocity >= 0.0f) UseSimSpaceSettings.MaxLinearVelocity = RBAN_SimSpaceOverride_MaxLinearVelocity;
 				if (RBAN_SimSpaceOverride_MaxAngularVelocity >= 0.0f) UseSimSpaceSettings.MaxAngularVelocity = RBAN_SimSpaceOverride_MaxAngularVelocity;
 				if (RBAN_SimSpaceOverride_MaxLinearAcceleration >= 0.0f) UseSimSpaceSettings.MaxLinearAcceleration = RBAN_SimSpaceOverride_MaxLinearAcceleration;
@@ -808,7 +818,8 @@ void FAnimNode_RigidBody::EvaluateSkeletalControl_AnyThread(FComponentSpacePoseC
 				SimulationSpaceAngularAcceleration);
 
 			PhysicsSimulation->SetSimulationSpaceSettings(
-				UseSimSpaceSettings.WorldAlpha,
+				(UseSimSpaceSettings.WorldAlpha > 0.0f),
+				UseSimSpaceSettings.DampingAlpha,
 				UseSimSpaceSettings.ExternalLinearDragV);
 
 			PhysicsSimulation->SetSolverSettings(
diff --git a/Engine/Source/Runtime/AnimGraphRuntime/Public/BoneControllers/AnimNode_RigidBody.h b/Engine/Source/Runtime/AnimGraphRuntime/Public/BoneControllers/AnimNode_RigidBody.h
--- a/Engine/Source/Runtime/AnimGraphRuntime/Public/BoneControllers/AnimNode_RigidBody.h
+++ b/Engine/Source/Runtime/AnimGraphRuntime/Public/BoneControllers/AnimNode_RigidBody.h
@@ -77,9 +77,18 @@ struct FSimSpaceSettings
 
 	// Multiplier on the Z-component of velocity and acceleration that is passed to the simulation. Usually from 0.0 to 1.0 to 
 	// reduce the effects of jumping and crouching on the simulation, but it can be higher than 1.0 if you need to exaggerate this motion for some reason.
+	// Should probably have been called "WorldAlphaScaleZ".
 	UPROPERTY(EditAnywhere, BlueprintReadWrite, Category = Settings, meta = (ClampMin = "0.0"))
 	float VelocityScaleZ;
 
+	// A muliplier to control how much of the simulation space movement is used to calculate the drag forces from Linear/Angular Damping in the Physics Asset.
+	// When DampingAlpha=1.0, Damping drag forces are equivalent to a world-space simulation. This is similar to air resistance.
+	// When DampingAlpha=0.0, Damping drag forces depend only on local-space body velocity and not on the simulation space velocity.
+	// It can be useful to set this to zero so that the Linear/Angular Damping settings on the BodyInstances do contribute to air resistance. Air resistance
+	// can be re-added in a controlled way using the ExternalLinearDrag setting below.
+	UPROPERTY(EditAnywhere, BlueprintReadWrite, Category = Settings, meta = (ClampMin = "0.0"))
+	float DampingAlpha;
+
 	// A clamp on the effective world-space velocity that is passed to the simulation. Units are cm/s. The default value effectively means "unlimited". It is not usually required to
 	// change this but you would reduce this to limit the effects of drag on the bodies in the simulation (if you have bodies that have LinearDrag set to non-zero in the physics asset). 
 	// Expected values in this case would be somewhat less than the usual velocities of your object which is commonly a few hundred for a character.
@@ -109,10 +118,11 @@ struct FSimSpaceSettings
 	float ExternalLinearDrag_DEPRECATED;
 #endif
 
-	// Additional linear drag applied to every body in addition to linear drag specified on them in the physics asset. 
-	// When combined with ExternalLinearVelocity, this can be used to add a temporary wind-blown effect without having to tune linear drag on 
-	// all the bodies in the physics asset. The result is that each body has a force equal to -ExternalLinearDragV * ExternalLinearVelocity applied to it, in 
-	// additional to all other forces. The vector is in simulation local space.
+	// Additional linear drag based on total body velocity. This is in addition to per-body linear damping in the physics asset (but see DampingAlpha to control that).
+	// (NOTE: The "V" suffix is to differentiate from the deprecated float property of the same name. It means "Vector" and not "Velocity").
+	//
+	// NOTE: ExternalLinearDragV is in simulation space, so if the RB AnimNode is set to Bone Space the ExternalLinearDragV.Z will be the drag in the
+	// Up direction of the selected bone.
 	UPROPERTY(EditAnywhere, BlueprintReadWrite, Category = Settings)
 	FVector ExternalLinearDragV;
 
diff --git a/Engine/Source/Runtime/Engine/Private/PhysicsEngine/ImmediatePhysics/ImmediatePhysicsChaos/ImmediatePhysicsSimulation_Chaos.cpp b/Engine/Source/Runtime/Engine/Private/PhysicsEngine/ImmediatePhysics/ImmediatePhysicsChaos/ImmediatePhysicsSimulation_Chaos.cpp
--- a/Engine/Source/Runtime/Engine/Private/PhysicsEngine/ImmediatePhysics/ImmediatePhysicsChaos/ImmediatePhysicsSimulation_Chaos.cpp
+++ b/Engine/Source/Runtime/Engine/Private/PhysicsEngine/ImmediatePhysics/ImmediatePhysicsChaos/ImmediatePhysicsSimulation_Chaos.cpp
@@ -805,12 +805,14 @@ namespace ImmediatePhysics_Chaos
 		Implementation->SimulationSpace.AngularVelocity = AngularVel;
 	}
 
-	void FSimulation::SetSimulationSpaceSettings(const FReal Alpha, const FVector& ExternalLinearEtherDrag)
+	void FSimulation::SetSimulationSpaceSettings(const bool bEnabled, const FReal DampingAlpha, const FVector& ExternalLinearEtherDrag)
 	{
 		using namespace Chaos;
 
 		FSimulationSpaceSettings& SimSpaceSettings = Implementation->Evolution.GetSimulationSpaceSettings();
-		SimSpaceSettings.Alpha = Alpha;
+		SimSpaceSettings.bEnabled = bEnabled;
+		SimSpaceSettings.LinearVelocityAlpha = DampingAlpha;
+		SimSpaceSettings.AngularVelocityAlpha = DampingAlpha;
 		SimSpaceSettings.ExternalLinearEtherDrag = ExternalLinearEtherDrag;
 		SimSpaceSettings.CentrifugalAlpha = ChaosImmediate_Evolution_SimSpaceCentrifugalAlpha;
 		SimSpaceSettings.CoriolisAlpha = ChaosImmediate_Evolution_SimSpaceCoriolisAlpha;
diff --git a/Engine/Source/Runtime/Engine/Public/Physics/ImmediatePhysics/ImmediatePhysicsChaos/ImmediatePhysicsSimulation_Chaos.h b/Engine/Source/Runtime/Engine/Public/Physics/ImmediatePhysics/ImmediatePhysicsChaos/ImmediatePhysicsSimulation_Chaos.h
--- a/Engine/Source/Runtime/Engine/Public/Physics/ImmediatePhysics/ImmediatePhysicsChaos/ImmediatePhysicsSimulation_Chaos.h
+++ b/Engine/Source/Runtime/Engine/Public/Physics/ImmediatePhysics/ImmediatePhysicsChaos/ImmediatePhysicsSimulation_Chaos.h
@@ -113,7 +113,8 @@ namespace ImmediatePhysics_Chaos
 			const FVector& AngularAcc);
 
 		ENGINE_API void SetSimulationSpaceSettings(
-			const FReal Alpha, 
+			const bool bEnabled, 
+			const FReal DampingAlpha,
 			const FVector& ExternalLinearEtherDrag);
 
 		ENGINE_API void SetSimulationSpaceSettings(
diff --git a/Engine/Source/Runtime/Experimental/Chaos/Private/Chaos/Evolution/PBDMinEvolution.cpp b/Engine/Source/Runtime/Experimental/Chaos/Private/Chaos/Evolution/PBDMinEvolution.cpp
--- a/Engine/Source/Runtime/Experimental/Chaos/Private/Chaos/Evolution/PBDMinEvolution.cpp
+++ b/Engine/Source/Runtime/Experimental/Chaos/Private/Chaos/Evolution/PBDMinEvolution.cpp
@@ -229,27 +229,43 @@ namespace Chaos
 					DV += Gravity * Dt;
 				}
 
+				FVec3 V = Particle.GetV();
+				FVec3 W = Particle.GetW();
+
 				// Moving and accelerating simulation frame
 				// https://en.wikipedia.org/wiki/Rotating_reference_frame
-				if (SimulationSpaceSettings.Alpha > 0.0f)
+				if (SimulationSpaceSettings.bEnabled)
 				{
-					const FVec3 CoriolisAcc = SimulationSpaceSettings.CoriolisAlpha * 2.0f * FVec3::CrossProduct(SimulationSpace.AngularVelocity, Particle.GetV());
-					const FVec3 CentrifugalAcc = SimulationSpaceSettings.CentrifugalAlpha * FVec3::CrossProduct(SimulationSpace.AngularVelocity, FVec3::CrossProduct(SimulationSpace.AngularVelocity, XCoM));
-					const FVec3 EulerAcc = SimulationSpaceSettings.EulerAlpha * FVec3::CrossProduct(SimulationSpace.AngularAcceleration, XCoM);
+					const FVec3 WxR = FVec3::CrossProduct(SimulationSpace.AngularVelocity, XCoM);
+					const FVec3 WxWxR = FVec3::CrossProduct(SimulationSpace.AngularVelocity, WxR);
+					const FVec3 WxV = FVec3::CrossProduct(SimulationSpace.AngularVelocity, V);
+					const FVec3 AxR = FVec3::CrossProduct(SimulationSpace.AngularAcceleration, XCoM);
+
+					// Calculate the phantom forces
+					const FVec3 CoriolisAcc = SimulationSpaceSettings.CoriolisAlpha * 2.0f * WxV;
+					const FVec3 CentrifugalAcc = SimulationSpaceSettings.CentrifugalAlpha * WxWxR;
+					const FVec3 EulerAcc = SimulationSpaceSettings.EulerAlpha * AxR;
 					const FVec3 LinearAcc = SimulationSpaceSettings.LinearAccelerationAlpha * SimulationSpace.LinearAcceleration;
 					const FVec3 AngularAcc = SimulationSpaceSettings.AngularAccelerationAlpha * SimulationSpace.AngularAcceleration;
-					const FVec3 LinearDragAcc = SimulationSpaceSettings.ExternalLinearEtherDrag * SimulationSpace.LinearVelocity;
-					DV -= SimulationSpaceSettings.Alpha * (LinearAcc + LinearDragAcc + CoriolisAcc + CentrifugalAcc + EulerAcc) * Dt;
-					DW -= SimulationSpaceSettings.Alpha * AngularAcc * Dt;
-					TargetV = -SimulationSpaceSettings.Alpha * SimulationSpaceSettings.LinearVelocityAlpha * SimulationSpace.LinearVelocity;
-					TargetW = -SimulationSpaceSettings.Alpha * SimulationSpaceSettings.AngularVelocityAlpha * SimulationSpace.AngularVelocity;
+					const FVec3 LinearDragAcc = SimulationSpaceSettings.ExternalLinearEtherDrag * (V + SimulationSpace.LinearVelocity + WxR);
+
+					// Apply the phantom forces
+					DV -= (LinearAcc + LinearDragAcc + CoriolisAcc + CentrifugalAcc + EulerAcc) * Dt;
+					DW -= AngularAcc * Dt;
+
+					// Pass on some part of world-space motion to the local-space drag calculation below.
+					// E.g., for LinearVelocityAlpha=1.0, and if LinearDrag is maxed out below, then the particle's local space velocity 
+					// would end up being exactly negative SimulationSpace.LinearVelocity. If LinearVelocityAlpha=0.0, then the 
+					// simulation-space velocity does not affect local drag.
+					TargetV = -SimulationSpaceSettings.LinearVelocityAlpha * SimulationSpace.LinearVelocity;
+					TargetW = -SimulationSpaceSettings.AngularVelocityAlpha * SimulationSpace.AngularVelocity;
 				}
 
 				// New velocity
-				const FReal LinearDrag = FMath::Min(FReal(1), Particle.LinearEtherDrag() * Dt);
-				const FReal AngularDrag = FMath::Min(FReal(1), Particle.AngularEtherDrag() * Dt);
-				const FVec3 V = FMath::Lerp(Particle.GetV() + DV, TargetV, LinearDrag);
-				FVec3 W = FMath::Lerp(Particle.GetW() + DW, TargetW, AngularDrag);
+				const FReal LinearDamping = FMath::Min(FReal(1), Particle.LinearEtherDrag() * Dt);
+				const FReal AngularDamping = FMath::Min(FReal(1), Particle.AngularEtherDrag() * Dt);
+				V = FMath::Lerp(V + DV, TargetV, LinearDamping);
+				W = FMath::Lerp(W + DW, TargetW, AngularDamping);
 
 				if (Particle.GyroscopicTorqueEnabled())
 				{
diff --git a/Engine/Source/Runtime/Experimental/Chaos/Public/Chaos/Evolution/SimulationSpace.h b/Engine/Source/Runtime/Experimental/Chaos/Public/Chaos/Evolution/SimulationSpace.h
--- a/Engine/Source/Runtime/Experimental/Chaos/Public/Chaos/Evolution/SimulationSpace.h
+++ b/Engine/Source/Runtime/Experimental/Chaos/Public/Chaos/Evolution/SimulationSpace.h
@@ -47,6 +47,7 @@ namespace Chaos
 		FSimulationSpaceSettings()
 			: MasterAlpha(0)
 			, Alpha(0)
+			, bEnabled(false)
 			, LinearAccelerationAlpha(1)
 			, CoriolisAlpha(1)
 			, CentrifugalAlpha(1)
@@ -63,12 +64,15 @@ PRAGMA_DISABLE_DEPRECATION_WARNINGS
 		FSimulationSpaceSettings& operator=(const FSimulationSpaceSettings&) = default;
 PRAGMA_ENABLE_DEPRECATION_WARNINGS
 
-		UE_DEPRECATED(5.1, "MasterAlpha is deprecated - please use Alpha")
+		UE_DEPRECATED(5.1, "MasterAlpha is deprecated - please use bEnabled")
 		FReal MasterAlpha;
 
-		// Global multipler on the effects of simulation space movement
+		UE_DEPRECATED(5.6, "Alpha is deprecated - please use bEnabled")
 		FReal Alpha;
 
+		// Enable or disable simulation space.
+		bool bEnabled;
+
 		// How much of the simulation frame's linear acceleration to pass onto the particles
 		FReal LinearAccelerationAlpha;
 
