Fix comment in RigidBody_AnimNode.h

diff --git a/Engine/Source/Runtime/AnimGraphRuntime/Public/BoneControllers/AnimNode_RigidBody.h b/Engine/Source/Runtime/AnimGraphRuntime/Public/BoneControllers/AnimNode_RigidBody.h
--- a/Engine/Source/Runtime/AnimGraphRuntime/Public/BoneControllers/AnimNode_RigidBody.h
+++ b/Engine/Source/Runtime/AnimGraphRuntime/Public/BoneControllers/AnimNode_RigidBody.h
@@ -84,8 +84,8 @@ struct FSimSpaceSettings
 	// A muliplier to control how much of the simulation space movement is used to calculate the drag forces from Linear/Angular Damping in the Physics Asset.
 	// When DampingAlpha=1.0, Damping drag forces are equivalent to a world-space simulation. This is similar to air resistance.
 	// When DampingAlpha=0.0, Damping drag forces depend only on local-space body velocity and not on the simulation space velocity.
-	// It can be useful to set this to zero so that the Linear/Angular Damping settings on the BodyInstances do contribute to air resistance. Air resistance
-	// can be re-added in a controlled way using the ExternalLinearDrag setting below.
+	// It can be useful to set this to zero so that the Linear/Angular Damping settings on the BodyInstances do not contribute to air resistance. 
+	// Air resistance can be re-added in a controlled way using the ExternalLinearDrag setting below.
 	UPROPERTY(EditAnywhere, BlueprintReadWrite, Category = Settings, meta = (ClampMin = "0.0"))
 	float DampingAlpha;
 
@@ -118,7 +118,7 @@ struct FSimSpaceSettings
 	float ExternalLinearDrag_DEPRECATED;
 #endif
 
-	// Additional linear drag based on total body velocity. This is in addition to per-body linear damping in the physics asset (but see DampingAlpha to control that).
+	// Additional linear drag applied to each body based on total body velocity. This is in addition to per-body linear damping in the physics asset (but see DampingAlpha to control that).
 	// (NOTE: The "V" suffix is to differentiate from the deprecated float property of the same name. It means "Vector" and not "Velocity").
 	//
 	// NOTE: ExternalLinearDragV is in simulation space, so if the RB AnimNode is set to Bone Space the ExternalLinearDragV.Z will be the drag in the
