Fix RBAN jittering with variable Dt
The previous velocities were not calculated with the previous Dt

diff --git a/Engine/Source/Runtime/AnimGraphRuntime/Private/BoneControllers/AnimNode_RigidBody.cpp b/Engine/Source/Runtime/AnimGraphRuntime/Private/BoneControllers/AnimNode_RigidBody.cpp
--- a/Engine/Source/Runtime/AnimGraphRuntime/Private/BoneControllers/AnimNode_RigidBody.cpp
+++ b/Engine/Source/Runtime/AnimGraphRuntime/Private/BoneControllers/AnimNode_RigidBody.cpp
@@ -241,6 +241,7 @@ FAnimNode_RigidBody::FAnimNode_RigidBody()
 	, PreviousComponentLinearVelocity(0.0f)
 	, PreviousSimulationSpaceTransform()
 	, PreviousPreviousSimulationSpaceTransform()
+	, PreviousDt(1.0f/30.0f)
 {
 }
 
@@ -384,6 +385,7 @@ void FAnimNode_RigidBody::InitSimulationSpace(
 	FTransform SimulationSpaceTransform = SpaceToWorldTransform(SimulationSpace, ComponentToWorld, BoneToComponent);
 	PreviousSimulationSpaceTransform = SimulationSpaceTransform;
 	PreviousPreviousSimulationSpaceTransform = SimulationSpaceTransform;
+	PreviousDt = 1.0f / 30.0f;
 }
 
 void FAnimNode_RigidBody::CalculateSimulationSpaceMotion(
@@ -417,12 +419,12 @@ void FAnimNode_RigidBody::CalculateSimulationSpaceMotion(
 	}
 
 	// World-space component linear velocity and acceleration
-	FVector PrevSpaceLinearVel = Chaos::FVec3::CalculateVelocity(PreviousPreviousSimulationSpaceTransform.GetTranslation(), PreviousSimulationSpaceTransform.GetTranslation(), Dt);
+	FVector PrevSpaceLinearVel = Chaos::FVec3::CalculateVelocity(PreviousPreviousSimulationSpaceTransform.GetTranslation(), PreviousSimulationSpaceTransform.GetTranslation(), PreviousDt);
 	SpaceLinearVel = Chaos::FVec3::CalculateVelocity(PreviousSimulationSpaceTransform.GetTranslation(), SimulationSpaceTransform.GetTranslation(), Dt);
 	SpaceLinearAcc = (SpaceLinearVel - PrevSpaceLinearVel) / Dt;
 
 	// World-space component angular velocity and acceleration
-	FVector PrevSpaceAngularVel = Chaos::FRotation3::CalculateAngularVelocity(PreviousPreviousSimulationSpaceTransform.GetRotation(), PreviousSimulationSpaceTransform.GetRotation(), Dt);
+	FVector PrevSpaceAngularVel = Chaos::FRotation3::CalculateAngularVelocity(PreviousPreviousSimulationSpaceTransform.GetRotation(), PreviousSimulationSpaceTransform.GetRotation(), PreviousDt);
 	SpaceAngularVel = Chaos::FRotation3::CalculateAngularVelocity(PreviousSimulationSpaceTransform.GetRotation(), SimulationSpaceTransform.GetRotation(), Dt);
 	SpaceAngularAcc = (SpaceAngularVel - PrevSpaceAngularVel) / Dt;
 
@@ -450,6 +452,7 @@ void FAnimNode_RigidBody::CalculateSimulationSpaceMotion(
 
 	PreviousPreviousSimulationSpaceTransform = PreviousSimulationSpaceTransform;
 	PreviousSimulationSpaceTransform = SimulationSpaceTransform;
+	PreviousDt = Dt;
 }
 
 
diff --git a/Engine/Source/Runtime/AnimGraphRuntime/Public/BoneControllers/AnimNode_RigidBody.h b/Engine/Source/Runtime/AnimGraphRuntime/Public/BoneControllers/AnimNode_RigidBody.h
--- a/Engine/Source/Runtime/AnimGraphRuntime/Public/BoneControllers/AnimNode_RigidBody.h
+++ b/Engine/Source/Runtime/AnimGraphRuntime/Public/BoneControllers/AnimNode_RigidBody.h
@@ -536,6 +536,7 @@ private:
 	// Used by the world-space to simulation-space motion transfer system in Component- or Bone-Space sims
 	FTransform PreviousSimulationSpaceTransform;
 	FTransform PreviousPreviousSimulationSpaceTransform;
+	float PreviousDt;
 
 #if ENABLE_LOW_LEVEL_MEM_TRACKER
 	FName OwningAssetPackageName;
