[AutoRTFM] Run Eigen computations in the open.

Eigen contains atomic variables and various threading logic. We aren't actually using the multithreaded capabilities of Eigen, but we appear to be hitting an atomic access on the Linux server.

I can't reproduce this issue locally on a Windows server, but this CL should address the Linux server issue based on the error I'm seeing in the server logs.

diff --git a/Engine/Source/Runtime/AnimGraphRuntime/Private/RBF/RBFInterpolator.cpp b/Engine/Source/Runtime/AnimGraphRuntime/Private/RBF/RBFInterpolator.cpp
--- a/Engine/Source/Runtime/AnimGraphRuntime/Private/RBF/RBFInterpolator.cpp
+++ b/Engine/Source/Runtime/AnimGraphRuntime/Private/RBF/RBFInterpolator.cpp
@@ -1,6 +1,7 @@
 // Copyright Epic Games, Inc. All Rights Reserved.
 
 #include "RBF/RBFInterpolator.h"
+#include "AutoRTFM/AutoRTFM.h"
 
 // Just to be sure, also added this in Eigen.Build.cs
 #ifndef EIGEN_MPL2_ONLY
@@ -27,54 +28,66 @@ bool FRBFInterpolatorBase::SetUpperKernel(
 {
 	using namespace Eigen;
 
-	MatrixXf FullKernel;
+	bool bSuccess = false;
 
-	// Construct the full kernel from the upper half calculated in the templated
-	// portion.
-	FullKernel = MatrixXf::Identity(Size, Size);
+	// The matrix logic below is self-contained and can run in the open without affecting
+	// transactional correctness.
+	// `Coeffs` is the only externally-visible alteration to program state.
+	// Allocating and zeroing Coeffs before we enter the open ensures that it can be rolled
+	// back safely if the transaction is aborted; we don't need a RecordOpenWrite below.
+	Coeffs.Init(0.0f, Size * Size);
 
-	for (int32 c = 0, i = 0; i < Size; i++)	
+	UE_AUTORTFM_OPEN
 	{
-		for (int32 j = i; j < Size; j++)
+		MatrixXf FullKernel;
+
+		// Construct the full kernel from the upper half calculated in the templated
+		// portion.
+		FullKernel = MatrixXf::Identity(Size, Size);
+
+		for (int32 c = 0, i = 0; i < Size; i++)	
 		{
-			FullKernel(i, j) = FullKernel(j, i) = UpperKernel[c++];
+			for (int32 j = i; j < Size; j++)
+			{
+				FullKernel(i, j) = FullKernel(j, i) = UpperKernel[c++];
+			}
 		}
-	}
 
-	// Usually the RBF formulation is computed by solving for:
-	// 
-	//   A * w = T
-	//
-	// Where A is the target kernel is a symmetric matrix containing the distance 
-	// between each node, w is the weights we want, and T is the target vector whose
-	// values we want to interpolate.
-	//
-	// However, in our case, we consider the activation of each node's output value
-	// to be a part of a N-dimensional vector, whose size is the same as the node count,
-	// and therefore, collectively, same as the target kernel's dimensions. 
-	// Each row is all zeros except for the activation value of 1.0 for a node at that
-	// node's index, effectively forming an identity matrix.
-	// 
-	// This allows us to reformulate the problem as:
-	//
-	//   A * w = I
-	//
-	// Or in other words:  
-	//
-	//   w = A^-1
-	//
-	// Eigen will now happily pick LU factorization with partial pivoting for 
-	// the inverse, which is blazingly fast.
+		// Usually the RBF formulation is computed by solving for:
+		// 
+		//   A * w = T
+		//
+		// Where A is the target kernel is a symmetric matrix containing the distance 
+		// between each node, w is the weights we want, and T is the target vector whose
+		// values we want to interpolate.
+		//
+		// However, in our case, we consider the activation of each node's output value
+		// to be a part of a N-dimensional vector, whose size is the same as the node count,
+		// and therefore, collectively, same as the target kernel's dimensions. 
+		// Each row is all zeros except for the activation value of 1.0 for a node at that
+		// node's index, effectively forming an identity matrix.
+		// 
+		// This allows us to reformulate the problem as:
+		//
+		//   A * w = I
+		//
+		// Or in other words:  
+		//
+		//   w = A^-1
+		//
+		// Eigen will now happily pick LU factorization with partial pivoting for 
+		// the inverse, which is blazingly fast.
 
-	Coeffs.Init(0.0f, Size * Size);
-	Map<MatrixXf> Result(Coeffs.GetData(), Size, Size);
+		Map<MatrixXf> Result(Coeffs.GetData(), Size, Size);
 
-	// If the matrix is non-invertible, return now and bail out.
-	float Det = FullKernel.determinant();
-	if (FMath::IsNearlyZero(Det))
-		return false;
-
-	Result = FullKernel.inverse();
+		// If the matrix is non-invertible, return now and bail out.
+		float Det = FullKernel.determinant();
+		if (!FMath::IsNearlyZero(Det))
+		{
+			Result = FullKernel.inverse();
+			bSuccess = true;
+		}
+	};
 
-	return true;
+	return bSuccess;
 }
