Anim Next: Fix for copy pose not propagating curves in anim bp if parent/source uses Anim Next

diff --git a/Engine/Source/Runtime/AnimGraphRuntime/Private/AnimNodes/AnimNode_CopyPoseFromMesh.cpp b/Engine/Source/Runtime/AnimGraphRuntime/Private/AnimNodes/AnimNode_CopyPoseFromMesh.cpp
--- a/Engine/Source/Runtime/AnimGraphRuntime/Private/AnimNodes/AnimNode_CopyPoseFromMesh.cpp
+++ b/Engine/Source/Runtime/AnimGraphRuntime/Private/AnimNodes/AnimNode_CopyPoseFromMesh.cpp
@@ -144,8 +144,13 @@ void FAnimNode_CopyPoseFromMesh::PreUpdate(const UAnimInstance* InAnimInstance)
 
 			if(bCopyCurves)
 			{
-				UAnimInstance* SourceAnimInstance = CurrentMeshComponent->GetAnimInstance();
-				if (SourceAnimInstance)
+				if (CurrentMeshComponent->bEnableAnimation == false)
+				{
+					// Assume this is using the anim next path. Only curves directly in the mesh component are valid.
+					// @TODO: Replace the path below with this one after validating performance and behavior.
+					SourceCurves.CopyFrom(CurrentMeshComponent->GetAnimCurves());
+				}
+				else if (UAnimInstance* SourceAnimInstance = CurrentMeshComponent->GetAnimInstance())
 				{
 					// Potential optimization/tradeoff: If we stored the curve results on the mesh component in non-editor scenarios, this would be
 					// much faster (but take more memory). As it is, we need to translate the map stored on the anim instance.
