Fixed  a crash on compilation of CalculateVelocityFromPositionHistory in Animation Blueprints

diff --git a/Engine/Source/Runtime/AnimGraphRuntime/Private/KismetAnimationLibrary.cpp b/Engine/Source/Runtime/AnimGraphRuntime/Private/KismetAnimationLibrary.cpp
--- a/Engine/Source/Runtime/AnimGraphRuntime/Private/KismetAnimationLibrary.cpp
+++ b/Engine/Source/Runtime/AnimGraphRuntime/Private/KismetAnimationLibrary.cpp
@@ -132,6 +132,12 @@ float UKismetAnimationLibrary::K2_CalculateVelocityFromPositionHistory(
 		History.LastIndex = 0;
 	}
 
+	if (History.Positions.Num() != History.Velocities.Num())
+	{
+		UE_LOG(LogAnimation, Warning, TEXT("Position History has a different number of position and velocity samples."));
+		return 0.f;
+	}
+
 	// append to the history until it's full and then loop around when filling it 
 	// to reuse the memory
 	if (History.Positions.Num() == 0)
