Fix link issues when using merged modules

#rnx

diff --git a/Engine/Source/Runtime/AnimGraphRuntime/Public/BoneControllers/AnimNode_SkeletalControlBase.h b/Engine/Source/Runtime/AnimGraphRuntime/Public/BoneControllers/AnimNode_SkeletalControlBase.h
--- a/Engine/Source/Runtime/AnimGraphRuntime/Public/BoneControllers/AnimNode_SkeletalControlBase.h
+++ b/Engine/Source/Runtime/AnimGraphRuntime/Public/BoneControllers/AnimNode_SkeletalControlBase.h
@@ -71,6 +71,8 @@ public:
 	{
 	}
 
+	virtual ~FAnimNode_SkeletalControlBase() {}
+
 public:
 #if WITH_EDITORONLY_DATA
 	// forwarded pose data from the wired node which current node's skeletal control is not applied yet
