[Animation] Keep the BlendMasks and LayerInfo in sync with the number of BlendPoses.

#rnx
#tests PIE in main, verified it fixed broken anim bps.

diff --git a/Engine/Source/Editor/AnimGraph/Private/AnimGraphNode_LayeredBoneBlend.cpp b/Engine/Source/Editor/AnimGraph/Private/AnimGraphNode_LayeredBoneBlend.cpp
--- a/Engine/Source/Editor/AnimGraph/Private/AnimGraphNode_LayeredBoneBlend.cpp
+++ b/Engine/Source/Editor/AnimGraph/Private/AnimGraphNode_LayeredBoneBlend.cpp
@@ -19,7 +19,7 @@
 UAnimGraphNode_LayeredBoneBlend::UAnimGraphNode_LayeredBoneBlend(const FObjectInitializer& ObjectInitializer)
 	: Super(ObjectInitializer)
 {
-	Node.AddPose();
+	Node.AddFirstPose();
 }
 
 FLinearColor UAnimGraphNode_LayeredBoneBlend::GetNodeTitleColor() const
@@ -47,18 +47,7 @@ void UAnimGraphNode_LayeredBoneBlend::PostEditChangeProperty(struct FPropertyCha
 		// If we  change blend modes, we need to resize our containers
 		FScopedTransaction Transaction(LOCTEXT("ChangeBlendMode", "Change Blend Mode"));
 		Modify();
-
-		const int32 NumPoses = Node.BlendPoses.Num();
-		if (Node.BlendMode == ELayeredBoneBlendMode::BlendMask)
-		{
-			Node.LayerSetup.Reset();
-			Node.BlendMasks.SetNum(NumPoses);
-		}
-		else
-		{
-			Node.BlendMasks.Reset();
-			Node.LayerSetup.SetNum(NumPoses);
-		}
+		Node.SyncBlendMasksAndLayers();
 
 		FBlueprintEditorUtils::MarkBlueprintAsStructurallyModified(GetBlueprint());
 	}
@@ -201,13 +190,10 @@ void UAnimGraphNode_LayeredBoneBlend::Serialize(FArchive& Ar)
 	Super::Serialize(Ar);
 
 	Ar.UsingCustomVersion(FUE5ReleaseStreamObjectVersion::GUID);
-
-	if (Ar.IsLoading() && Ar.CustomVer(FUE5ReleaseStreamObjectVersion::GUID) < FUE5ReleaseStreamObjectVersion::AnimLayeredBoneBlendMasks)
+	
+	if (Ar.IsLoading() && Ar.CustomVer(FUE5ReleaseStreamObjectVersion::GUID) < FUE5ReleaseStreamObjectVersion::AnimLayeredBoneBlendMasksFix)
 	{
-		if (Node.BlendMode == ELayeredBoneBlendMode::BlendMask && Node.BlendMasks.Num() != Node.BlendPoses.Num())
-		{
-			Node.BlendMasks.SetNum(Node.BlendPoses.Num());
-		}
+		Node.SyncBlendMasksAndLayers();
 	}
 }
 
diff --git a/Engine/Source/Runtime/AnimGraphRuntime/Private/AnimNodes/AnimNode_LayeredBoneBlend.cpp b/Engine/Source/Runtime/AnimGraphRuntime/Private/AnimNodes/AnimNode_LayeredBoneBlend.cpp
--- a/Engine/Source/Runtime/AnimGraphRuntime/Private/AnimNodes/AnimNode_LayeredBoneBlend.cpp
+++ b/Engine/Source/Runtime/AnimGraphRuntime/Private/AnimNodes/AnimNode_LayeredBoneBlend.cpp
@@ -32,6 +32,20 @@ void FAnimNode_LayeredBoneBlend::Initialize_AnyThread(const FAnimationInitialize
 	}
 }
 
+void FAnimNode_LayeredBoneBlend::SyncBlendMasksAndLayers()
+{
+	if (BlendMode == ELayeredBoneBlendMode::BlendMask)
+	{
+		BlendMasks.SetNum(BlendPoses.Num());
+		LayerSetup.Reset();
+	}
+	else
+	{
+		BlendMasks.Reset();
+		LayerSetup.SetNum(BlendPoses.Num());
+	}
+}
+
 void FAnimNode_LayeredBoneBlend::RebuildPerBoneBlendWeights(const USkeleton* InSkeleton)
 {
 	if (InSkeleton)
@@ -300,4 +314,4 @@ void FAnimNode_LayeredBoneBlend::SetBlendMask(int32 InPoseIndex, UBlendProfile*
 	BlendMasks[InPoseIndex] = InBlendMask;
 
 	InvalidatePerBoneBlendWeights();
-}
\ No newline at end of file
+}
diff --git a/Engine/Source/Runtime/AnimGraphRuntime/Public/AnimNodes/AnimNode_LayeredBoneBlend.h b/Engine/Source/Runtime/AnimGraphRuntime/Public/AnimNodes/AnimNode_LayeredBoneBlend.h
--- a/Engine/Source/Runtime/AnimGraphRuntime/Public/AnimNodes/AnimNode_LayeredBoneBlend.h
+++ b/Engine/Source/Runtime/AnimGraphRuntime/Public/AnimNodes/AnimNode_LayeredBoneBlend.h
@@ -81,7 +81,7 @@ protected:
 
 public:
 	/*
- 	 * Max LOD that this node is allowed to run
+	 * Max LOD that this node is allowed to run
 	 * For example if you have LODThreshold to be 2, it will run until LOD 2 (based on 0 index)
 	 * when the component LOD becomes 3, it will stop update/evaluate
 	 * currently transition would be issue and that has to be re-visited
@@ -137,15 +137,7 @@ public:
 	{
 		BlendWeights.Add(1.f);
 		BlendPoses.AddDefaulted();
-
-		if (BlendMode == ELayeredBoneBlendMode::BlendMask) 
-		{ 
-			BlendMasks.Add(nullptr); 
-		}
-		else /*if (BlendMode::BranchFilter)*/ 
-		{ 
-			LayerSetup.AddDefaulted(); 
-		}
+		SyncBlendMasksAndLayers();
 	}
 
 	void RemovePose(int32 PoseIndex)
@@ -155,13 +147,15 @@ public:
 
 		if (BlendMasks.IsValidIndex(PoseIndex)) 
 		{ 
-			BlendMasks.RemoveAt(PoseIndex); 
+			BlendMasks.RemoveAt(PoseIndex);
 		}
 
 		if (LayerSetup.IsValidIndex(PoseIndex)) 
 		{ 
 			LayerSetup.RemoveAt(PoseIndex); 
 		}
+		// just in case, we've seen problems.
+		SyncBlendMasksAndLayers();
 	}
 
 	// Set the blend mask for the specified input pose
@@ -172,8 +166,22 @@ public:
 	
 	// Invalidates the cached bone data so it is recalculated the next time this node is updated
 	void InvalidateCachedBoneData() { RequiredBonesSerialNumber = 0; }
-	
+
 private:
+
+	// just for the graph constructor to call.
+	void AddFirstPose()
+	{
+		check(BlendWeights.IsEmpty());
+		BlendWeights.Add(1.f);
+		BlendPoses.AddDefaulted();
+		check(BlendMode == ELayeredBoneBlendMode::BranchFilter);
+		LayerSetup.AddDefaulted();
+	}
+
+	// Synchronize the number of BlendMasks or Layers with the number of BlendPoses.
+	ANIMGRAPHRUNTIME_API void SyncBlendMasksAndLayers();
+	
 	// Rebuild cache per bone blend weights from the skeleton
 	ANIMGRAPHRUNTIME_API void RebuildPerBoneBlendWeights(const USkeleton* InSkeleton);
 
diff --git a/Engine/Source/Runtime/Core/Public/UObject/UE5ReleaseStreamObjectVersions.inl b/Engine/Source/Runtime/Core/Public/UObject/UE5ReleaseStreamObjectVersions.inl
--- a/Engine/Source/Runtime/Core/Public/UObject/UE5ReleaseStreamObjectVersions.inl
+++ b/Engine/Source/Runtime/Core/Public/UObject/UE5ReleaseStreamObjectVersions.inl
@@ -162,5 +162,8 @@ UE5_RELEASE_VERSION(SinglePrecisonParticleDataPT, 50)
 //Orthographic Near and Far Plane Auto-resolve enabled by default
 UE5_RELEASE_VERSION(OrthographicAutoNearFarPlane, 51)
 
+// Fix a bug where BlendMask counts could get out of sync with BlendPose counts.
+UE5_RELEASE_VERSION(AnimLayeredBoneBlendMasksFix, 52)
+
 // -----<new versions can be added above this line>-------------------------------------------------
 #undef UE5_RELEASE_VERSION
\ No newline at end of file
