Fixing blend space length fetch crash when animation fails to load.
#rnx

diff --git a/Engine/Source/Runtime/AnimGraphRuntime/Private/AnimNodes/AnimNode_BlendSpacePlayer.cpp b/Engine/Source/Runtime/AnimGraphRuntime/Private/AnimNodes/AnimNode_BlendSpacePlayer.cpp
--- a/Engine/Source/Runtime/AnimGraphRuntime/Private/AnimNodes/AnimNode_BlendSpacePlayer.cpp
+++ b/Engine/Source/Runtime/AnimGraphRuntime/Private/AnimNodes/AnimNode_BlendSpacePlayer.cpp
@@ -44,7 +44,10 @@ float FAnimNode_BlendSpacePlayerBase::GetCurrentAssetLength() const
 		if (CurrentBlendSpace != nullptr)
 		{
 			const FBlendSample& Sample = CurrentBlendSpace->GetBlendSample(HighestWeightedSample->SampleDataIndex);
-			return Sample.Animation->GetPlayLength();
+			if (ensure(Sample.Animation != nullptr))
+			{
+				return Sample.Animation->GetPlayLength();
+			}
 		}
 	}
 
