Fix for visualization of what the blend weight of the animation is in Animation Insights for Random Sequence Player.  Now it takes the overall blendweight of the animation into account as well.

diff --git a/Engine/Source/Runtime/AnimGraphRuntime/Private/AnimNodes/AnimNode_RandomPlayer.cpp b/Engine/Source/Runtime/AnimGraphRuntime/Private/AnimNodes/AnimNode_RandomPlayer.cpp
--- a/Engine/Source/Runtime/AnimGraphRuntime/Private/AnimNodes/AnimNode_RandomPlayer.cpp
+++ b/Engine/Source/Runtime/AnimGraphRuntime/Private/AnimNodes/AnimNode_RandomPlayer.cpp
@@ -237,7 +237,7 @@ void FAnimNode_RandomPlayer::Update_AnyThread(const FAnimationUpdateContext& Con
 		NextData = &GetPlayData(ERandomDataIndexType::Next);
 	}
 
-	FAnimTickRecord TickRecord(CurrentData->Entry->Sequence, true, CurrentData->PlayRate, false, CurrentData->BlendWeight, CurrentData->CurrentPlayTime, CurrentData->MarkerTickRecord);
+	FAnimTickRecord TickRecord(CurrentData->Entry->Sequence, true, CurrentData->PlayRate, false, CurrentData->BlendWeight * BlendWeight, CurrentData->CurrentPlayTime, CurrentData->MarkerTickRecord);
 	TickRecord.DeltaTimeRecord = &CurrentData->DeltaTimeRecord;
 	TickRecord.GatherContextData(Context);
 
@@ -248,7 +248,7 @@ void FAnimNode_RandomPlayer::Update_AnyThread(const FAnimationUpdateContext& Con
 
 	if (FAnimationRuntime::HasWeight(NextData->BlendWeight))
 	{
-		FAnimTickRecord NextTickRecord(NextData->Entry->Sequence, true, NextData->PlayRate, false, NextData->BlendWeight, NextData->CurrentPlayTime, NextData->MarkerTickRecord);
+		FAnimTickRecord NextTickRecord(NextData->Entry->Sequence, true, NextData->PlayRate, false, NextData->BlendWeight * BlendWeight, NextData->CurrentPlayTime, NextData->MarkerTickRecord);
 		NextTickRecord.DeltaTimeRecord = &NextData->DeltaTimeRecord;
 		NextTickRecord.GatherContextData(Context);
 
