Added an input pin option for AnimNode_SpringBone to use an override value for owner velocity.

diff --git a/Engine/Source/Runtime/AnimGraphRuntime/Private/BoneControllers/AnimNode_SpringBone.cpp b/Engine/Source/Runtime/AnimGraphRuntime/Private/BoneControllers/AnimNode_SpringBone.cpp
--- a/Engine/Source/Runtime/AnimGraphRuntime/Private/BoneControllers/AnimNode_SpringBone.cpp
+++ b/Engine/Source/Runtime/AnimGraphRuntime/Private/BoneControllers/AnimNode_SpringBone.cpp
@@ -33,6 +33,8 @@ FAnimNode_SpringBone::FAnimNode_SpringBone()
 	, bRotateY(false)
 	, bRotateZ(false)
 	, bHadValidStrength(false)
+	, bOverrideOwnerVelocity(false)
+	, OwnerVelocityOverride(FVector::ZeroVector)
 {
 }
 
@@ -246,16 +248,24 @@ void FAnimNode_SpringBone::PreUpdate(const UAnimInstance* InAnimInstance)
 			check(World->GetWorldSettings());
 			TimeDilation = World->GetWorldSettings()->GetEffectiveTimeDilation();
 
-			AActor* SkelOwner = SkelComp->GetOwner();
-			if (SkelComp->GetAttachParent() != nullptr && (SkelOwner == nullptr))
+			if (bOverrideOwnerVelocity)
 			{
-				SkelOwner = SkelComp->GetAttachParent()->GetOwner();
-				OwnerVelocity = SkelOwner->GetVelocity();
+				OwnerVelocity = OwnerVelocityOverride;
 			}
 			else
 			{
-				OwnerVelocity = FVector::ZeroVector;
+				AActor* SkelOwner = SkelComp->GetOwner();
+				if (SkelComp->GetAttachParent() != nullptr && (SkelOwner == nullptr))
+				{
+					SkelOwner = SkelComp->GetAttachParent()->GetOwner();
+					OwnerVelocity = SkelOwner->GetVelocity();
+				}
+				else
+				{
+					OwnerVelocity = FVector::ZeroVector;
+				}
 			}
+
 		}
 	}
 }
diff --git a/Engine/Source/Runtime/AnimGraphRuntime/Public/BoneControllers/AnimNode_SpringBone.h b/Engine/Source/Runtime/AnimGraphRuntime/Public/BoneControllers/AnimNode_SpringBone.h
--- a/Engine/Source/Runtime/AnimGraphRuntime/Public/BoneControllers/AnimNode_SpringBone.h
+++ b/Engine/Source/Runtime/AnimGraphRuntime/Public/BoneControllers/AnimNode_SpringBone.h
@@ -100,6 +100,13 @@ struct FAnimNode_SpringBone : public FAnimNode_SkeletalControlBase
 	/** Did we have a non-zero ControlStrength last frame. */
 	uint8 bHadValidStrength : 1;
 
+	/** Option to override owner velocity used by spring controller. */
+	UPROPERTY(EditAnywhere, Category = OwnerVelocity, Meta = (PinHiddenByDefault))
+	uint8 bOverrideOwnerVelocity : 1;
+
+	UPROPERTY(EditAnywhere, Category = OwnerVelocity, Meta = (EditCondition = "bOverrideOwnerVelocity", PinHiddenByDefault))
+	FVector OwnerVelocityOverride;
+
 public:
 	ANIMGRAPHRUNTIME_API FAnimNode_SpringBone();
 
