Add scale and external rotational forces support to AnimDynamics

diff --git a/Engine/Source/Runtime/AnimGraphRuntime/Private/BoneControllers/AnimNode_AnimDynamics.cpp b/Engine/Source/Runtime/AnimGraphRuntime/Private/BoneControllers/AnimNode_AnimDynamics.cpp
--- a/Engine/Source/Runtime/AnimGraphRuntime/Private/BoneControllers/AnimNode_AnimDynamics.cpp
+++ b/Engine/Source/Runtime/AnimGraphRuntime/Private/BoneControllers/AnimNode_AnimDynamics.cpp
@@ -22,12 +22,21 @@ CSV_DECLARE_CATEGORY_MODULE_EXTERN(ENGINE_API, Animation);
 TAutoConsoleVariable<int32> CVarRestrictLod(TEXT("p.AnimDynamicsRestrictLOD"), -1, TEXT("Forces anim dynamics to be enabled for only a specified LOD, -1 to enable on all LODs."));
 TAutoConsoleVariable<int32> CVarLODThreshold(TEXT("p.AnimDynamicsLODThreshold"), -1, TEXT("Max LOD that anim dynamics is allowed to run on. Provides a global threshold that overrides per-node the LODThreshold property. -1 means no override."), ECVF_Scalability);
 TAutoConsoleVariable<int32> CVarEnableDynamics(TEXT("p.AnimDynamics"), 1, TEXT("Enables/Disables anim dynamics node updates."), ECVF_Scalability);
+TAutoConsoleVariable<int32> CVarEnableScaling(TEXT("p.AnimDynamics.ScaleSupport"), 1, TEXT("Enables/Disables scaling support for anim dynamics nodes."));
 TAutoConsoleVariable<int32> CVarEnableAdaptiveSubstep(TEXT("p.AnimDynamicsAdaptiveSubstep"), 0, TEXT("Enables/disables adaptive substepping. Adaptive substepping will substep the simulation when it is necessary and maintain a debt buffer for time, always trying to utilise as much time as possible."));
 TAutoConsoleVariable<int32> CVarAdaptiveSubstepNumDebtFrames(TEXT("p.AnimDynamicsNumDebtFrames"), 5, TEXT("Number of frames to maintain as time debt when using adaptive substepping, this should be at least 1 or the time debt will never be cleared."));
 TAutoConsoleVariable<int32> CVarEnableWind(TEXT("p.AnimDynamicsWind"), 1, TEXT("Enables/Disables anim dynamics wind forces globally."), ECVF_Scalability);
 TAutoConsoleVariable<float> CVarComponentAppliedLinearAccClampOverride(TEXT("p.AnimDynamics.ComponentAppliedLinearAccClampOverride"), -1.0f, TEXT("Override the per asset setting for all axis (X,Y & Z) of ComponentAppliedLinearAccClamp for all Anim Dynamics Nodes. Negative values are ignored."));
 TAutoConsoleVariable<float> CVarGravityScale(TEXT("p.AnimDynamics.GravityScale"), 1.0f, TEXT("Multiplies the default gravity and the gravity override on all Anim Dynamics Nodes."));
 
+// FSimSpaceSettingsWithControl forced overrides for testing
+bool bAnimDynamics_SimSpace_EnableOverride = false;
+FAnimPhysSimSpaceSettings AnimDynamics_SimSpaceOverride;
+FAutoConsoleVariableRef CVarAnimDynamicsSpaceOverride(TEXT("p.AnimDynamics.Space.Override"), bAnimDynamics_SimSpace_EnableOverride, TEXT("Force-enable the advanced simulation space movement forces"), ECVF_Default);
+FAutoConsoleVariableRef CVarAnimDynamicsSpaceAngularAlpha(TEXT("p.AnimDynamics.Space.SimSpaceAngularAlpha"), AnimDynamics_SimSpaceOverride.SimSpaceAngularAlpha, TEXT("AnimDynamics SimSpaceSettings overrides"), ECVF_Default);
+FAutoConsoleVariableRef CVarAnimDynamicsSpaceMaxCompAngVel(TEXT("p.AnimDynamics.Space.MaxAngularVelocity"), AnimDynamics_SimSpaceOverride.MaxAngularVelocity, TEXT("AnimDynamics SimSpaceSettings overrides"), ECVF_Default);
+FAutoConsoleVariableRef CVarAnimDynamicsSpaceMaxCompAngAcc(TEXT("p.AnimDynamics.Space.MaxAngularAcceleration"), AnimDynamics_SimSpaceOverride.MaxAngularAcceleration, TEXT("AnimDynamics SimSpaceSettings overrides"), ECVF_Default);
+
 // FindChainBones
 // 
 // Returns an ordered list of the names of all the bones in a chain between BeginBoneName and EndBoneName.
@@ -102,7 +111,7 @@ void FAnimNode_AnimDynamics::DrawBodies(FComponentSpacePoseContext& InContext, c
 			continue;
 		}
 
-		FTransform Transform(Body.Pose.Orientation, Body.Pose.Position);
+		FTransform Transform(Body.Pose.Orientation, Body.Pose.Position, FVector(Body.Pose.Scale));
 		Transform = GetComponentSpaceTransformFromSimSpace(SimulationSpace, InContext, Transform);
 
 		Proxy->AnimDrawDebugCoordinateSystem(Transform.GetTranslation(), Transform.Rotator(), 2.0f, false, -1.0f, 0.15f);
@@ -164,6 +173,14 @@ void FAnimNode_AnimDynamics::DrawBodies(FComponentSpacePoseContext& InContext, c
 }
 #endif
 
+FAnimPhysSimSpaceSettings::FAnimPhysSimSpaceSettings()
+	: SimSpaceAngularAlpha(0)
+	, MaxAngularVelocity(10000)
+	, MaxAngularAcceleration(10000)
+	, ExternalAngularVelocity(FVector::ZeroVector)
+{
+}
+
 FAnimNode_AnimDynamics::FAnimNode_AnimDynamics()
 : LinearDampingOverride(0.0f)
 , AngularDampingOverride(0.0f)
@@ -214,6 +231,7 @@ FAnimNode_AnimDynamics::FAnimNode_AnimDynamics()
 	ComponentAppliedLinearAccClamp = FVector(100000, 100000, 100000);
 
 	PreviousComponentLinearVelocity = FVector::ZeroVector;
+	PreviousSimSpaceAngularVelocity = FVector::ZeroVector;
 
 	PhysicsBodyDefinitions.Add(FAnimPhysBodyDefinition());
 }
@@ -234,6 +252,7 @@ void FAnimNode_AnimDynamics::Initialize_AnyThread(const FAnimationInitializeCont
 
 	PreviousCompWorldSpaceTM = Context.AnimInstanceProxy->GetComponentTransform();
 	PreviousActorWorldSpaceTM = Context.AnimInstanceProxy->GetActorTransform();
+	PreviousSimSpaceTM = PreviousCompWorldSpaceTM;
 
 	NextTimeStep = 0.0f;
 	TimeDebt = 0.0f;
@@ -261,6 +280,111 @@ bool FAnimNode_AnimDynamics::IsAnimDynamicsSystemEnabledFor(int32 InLOD)
 	// GetLODThreshold per node
 	return (CVarEnableDynamics.GetValueOnAnyThread() == 1 && bEnabledForLod);
 }
+
+FTransform FAnimNode_AnimDynamics::GetSimToWorldSpaceTransform(AnimPhysSimSpaceType SimSpace, FComponentSpacePoseContext& Output) const
+{
+	FTransform OutTransform;
+
+	switch (SimSpace)
+	{
+	case AnimPhysSimSpaceType::Component:
+	{
+		OutTransform = Output.AnimInstanceProxy->GetComponentTransform();
+		
+		break;
+	}
+
+	case AnimPhysSimSpaceType::Actor:
+	{
+		OutTransform = Output.AnimInstanceProxy->GetActorTransform();
+
+		break;
+	}
+
+	case AnimPhysSimSpaceType::RootRelative:
+	{
+		const FBoneContainer& RequiredBones = Output.Pose.GetPose().GetBoneContainer();
+
+		FCompactPoseBoneIndex RootBoneCompactIndex(0);
+
+		FTransform RelativeBoneTransform = Output.Pose.GetComponentSpaceTransform(RootBoneCompactIndex);
+		RelativeBoneTransform = Output.AnimInstanceProxy->GetComponentTransform() * RelativeBoneTransform;
+		OutTransform = RelativeBoneTransform;
+
+		break;
+	}
+
+	case AnimPhysSimSpaceType::BoneRelative:
+	{
+		const FBoneContainer& RequiredBones = Output.Pose.GetPose().GetBoneContainer();
+		if (RelativeSpaceBone.IsValidToEvaluate(RequiredBones))
+		{
+			FTransform RelativeBoneTransform = Output.Pose.GetComponentSpaceTransform(RelativeSpaceBone.GetCompactPoseIndex(RequiredBones));
+			RelativeBoneTransform = Output.AnimInstanceProxy->GetComponentTransform() * RelativeBoneTransform;
+			OutTransform = RelativeBoneTransform;
+		}
+
+		break;
+	}
+	case AnimPhysSimSpaceType::World:
+	{
+		break;
+	}
+
+	default:
+		break;
+	}
+
+	return OutTransform;
+}
+
+void FAnimNode_AnimDynamics::CalculateSimulationSpaceMotion(
+	AnimPhysSimSpaceType        Space,
+	FAnimPhysSimSpaceSettings	Settings,
+	FComponentSpacePoseContext& Output,
+	const float              Dt,
+	FVector& SpaceAngularVel,
+	FVector& SpaceAngularAcc)
+{
+	SpaceAngularVel = FVector::ZeroVector;
+	SpaceAngularAcc = FVector::ZeroVector;
+
+	// If the system is disabled, nothing else to do
+	if ((Settings.SimSpaceAngularAlpha == 0.0f) || (Dt < UE_SMALL_NUMBER))
+	{
+		return;
+	}
+
+	// If the simulation is in world space, the simulation space is stationary
+	if (Space == AnimPhysSimSpaceType::World)
+	{
+		SpaceAngularVel = SimSpaceSettings.ExternalAngularVelocity;
+		return;
+	}
+
+	const FTransform CurrentSimSpaceTM = GetSimToWorldSpaceTransform(SimulationSpace, Output);
+
+	// World-space component angular velocity and acceleration
+	SpaceAngularVel = Chaos::FRotation3::CalculateAngularVelocity(PreviousSimSpaceTM.GetRotation(), CurrentSimSpaceTM.GetRotation(), Dt);
+	SpaceAngularAcc = (SpaceAngularVel - PreviousSimSpaceAngularVelocity) / Dt;
+
+	PreviousSimSpaceAngularVelocity = SpaceAngularVel;
+	PreviousSimSpaceTM = CurrentSimSpaceTM;
+
+	// Clamped world-space motion of the simulation space
+	SpaceAngularVel = SpaceAngularVel.GetClampedToMaxSize(Settings.MaxAngularVelocity) + Settings.ExternalAngularVelocity;
+	SpaceAngularAcc = SpaceAngularAcc.GetClampedToMaxSize(Settings.MaxAngularAcceleration);
+
+	// Transform world-space motion into simulation space
+	SpaceAngularVel = CurrentSimSpaceTM.InverseTransformVectorNoScale(SpaceAngularVel);
+	SpaceAngularAcc = CurrentSimSpaceTM.InverseTransformVectorNoScale(SpaceAngularAcc);
+
+	// Apply WorldAlpha to simulation space rotation (usually to reduce it)
+	SpaceAngularVel *= Settings.SimSpaceAngularAlpha;
+	SpaceAngularAcc *= Settings.SimSpaceAngularAlpha;
+
+}
+
 void FAnimNode_AnimDynamics::EvaluateSkeletalControl_AnyThread(FComponentSpacePoseContext& Output, TArray<FBoneTransform>& OutBoneTransforms)
 {
 	DECLARE_SCOPE_HIERARCHICAL_COUNTER_ANIMNODE(EvaluateSkeletalControl_AnyThread)
@@ -269,6 +393,8 @@ void FAnimNode_AnimDynamics::EvaluateSkeletalControl_AnyThread(FComponentSpacePo
 
 	if (IsAnimDynamicsSystemEnabledFor(Output.AnimInstanceProxy->GetLODLevel()))
 	{
+		const int bScalingEnabled = CVarEnableScaling.GetValueOnAnyThread();
+
 		if(LastSimSpace != SimulationSpace)
 		{
 			// Our sim space has been changed since our last update, we need to convert all of our
@@ -284,6 +410,61 @@ void FAnimNode_AnimDynamics::EvaluateSkeletalControl_AnyThread(FComponentSpacePo
 			InitTeleportType = ETeleportType::None;
 		}
 
+		if (bScalingEnabled)
+		{
+			const FTransform SimWorldTransform = GetSimToWorldSpaceTransform(SimulationSpace, Output);
+
+			if(SimWorldTransform.GetScale3D() != PreviousSimSpaceTM.GetScale3D())
+			{
+				// If the scaling of the sim space has changed, update the bodies to the proper scale
+				const FBoneContainer& BoneContainer = Output.Pose.GetPose().GetBoneContainer();
+				for (int32 ActiveIndex : ActiveBoneIndices)
+				{
+					const FAnimPhysBodyDefinition& PhysicsBodyDef = PhysicsBodyDefinitions[ActiveIndex];
+					const FBoneReference& CurrentBoneRef = PhysicsBodyDef.BoundBone;
+
+					// If our bone isn't valid, move on
+					if (!CurrentBoneRef.IsValidToEvaluate(BoneContainer))
+					{
+						continue;
+					}
+
+					FAnimPhysLinkedBody& ChainBody = Bodies[ActiveIndex];
+					FAnimPhysRigidBody& RigidBody = Bodies[ActiveIndex].RigidBody.PhysBody;
+
+					FCompactPoseBoneIndex BoneIndex = CurrentBoneRef.GetCompactPoseIndex(BoneContainer);
+					FTransform BoundBoneTransform = GetBoneTransformInSimSpace(Output, BoneIndex);
+
+					float Scale = BoundBoneTransform.GetScale3D().X;
+					if (RigidBody.Pose.Scale != Scale)
+					{
+						FAnimPhys::ScaleRigidBodyMass(&RigidBody, Scale/RigidBody.Pose.Scale);
+
+						TArray<FAnimPhysShape> BodyShapes;
+
+						BodyShapes.Add(FAnimPhysShape::MakeBox(PhysicsBodyDef.BoxExtents * Scale));
+						RigidBody.Pose.Scale = Scale;
+
+						switch (RigidBody.CollisionType)
+						{
+						case AnimPhysCollisionType::CustomSphere:
+							RigidBody.SphereCollisionRadius = static_cast<float>(PhysicsBodyDef.SphereCollisionRadius * Scale);
+							break;
+						case AnimPhysCollisionType::InnerSphere:
+							RigidBody.SphereCollisionRadius = static_cast<float>(PhysicsBodyDef.BoxExtents.GetAbsMin() * Scale / 2.0);
+							break;
+						case AnimPhysCollisionType::OuterSphere:
+							RigidBody.SphereCollisionRadius = static_cast<float>(PhysicsBodyDef.BoxExtents.GetAbsMax() * Scale / 2.0);
+							break;
+						default:
+							break;
+						}
+						RigidBody.Shapes = BodyShapes;
+					}
+				}
+			}
+		}
+
 		const FBoneContainer& RequiredBones = Output.Pose.GetPose().GetBoneContainer();
 		while(BodiesToReset.Num() > 0)
 		{
@@ -295,6 +476,7 @@ void FAnimNode_AnimDynamics::EvaluateSkeletalControl_AnyThread(FComponentSpacePo
 
 				PhysBody.Pose.Position = BoneTransform.GetTranslation();
 				PhysBody.Pose.Orientation = BoneTransform.GetRotation();
+				PhysBody.Pose.Scale = BoneTransform.GetScale3D().X;
 				PhysBody.LinearMomentum = FVector::ZeroVector;
 				PhysBody.AngularMomentum = FVector::ZeroVector;
 			}
@@ -325,27 +507,64 @@ void FAnimNode_AnimDynamics::EvaluateSkeletalControl_AnyThread(FComponentSpacePo
 			FVector ComponentLinearAcc(0.0f);
 			FVector SimSpaceGravityOverride = GravityOverride;
 
-			if (SimulationSpace != AnimPhysSimSpaceType::World)
+			FVector SimFrameAngularVelocity;
+			FVector SimFrameAngularAcc;
+
+			FAnimPhysSimSpaceSettings UseSimSpaceSettings = SimSpaceSettings;
+			if (bAnimDynamics_SimSpace_EnableOverride)
 			{
-				FTransform CurrentTransform = Output.AnimInstanceProxy->GetComponentTransform();
+				if (AnimDynamics_SimSpaceOverride.SimSpaceAngularAlpha >= 0.0f)
+				{
+					UseSimSpaceSettings.SimSpaceAngularAlpha = AnimDynamics_SimSpaceOverride.SimSpaceAngularAlpha;
+				}
+				if (AnimDynamics_SimSpaceOverride.MaxAngularVelocity >= 0.0f)
+				{
+					UseSimSpaceSettings.MaxAngularVelocity = AnimDynamics_SimSpaceOverride.MaxAngularVelocity;
+				}
+				if (AnimDynamics_SimSpaceOverride.MaxAngularAcceleration >= 0.0f)
+				{ 
+					UseSimSpaceSettings.MaxAngularAcceleration = AnimDynamics_SimSpaceOverride.MaxAngularAcceleration;
+				}
+			}
+
+			CalculateSimulationSpaceMotion(SimulationSpace, SimSpaceSettings, Output, NextTimeStep, SimFrameAngularVelocity, SimFrameAngularAcc);
 
+			if (SimulationSpace != AnimPhysSimSpaceType::World)
+			{
 				// Transform Gravity Override into simulation space
 				if (bUseGravityOverride && !bGravityOverrideInSimSpace)
 				{
 					SimSpaceGravityOverride = TransformWorldVectorToSimSpace(Output, SimSpaceGravityOverride);
 				}
 
+				FTransform CurrentTransform = Output.AnimInstanceProxy->GetComponentTransform();
+
 				// Calc linear velocity
-				const FVector ComponentDeltaLocation = CurrentTransform.GetTranslation() - PreviousCompWorldSpaceTM.GetTranslation();
-				const FVector ComponentLinearVelocity = ComponentDeltaLocation / NextTimeStep;
+				const FVector ComponentDeltaLocation =  CurrentTransform.GetTranslation() - PreviousCompWorldSpaceTM.GetTranslation();
+				const FVector ComponentLinearVelocity = (ComponentDeltaLocation / NextTimeStep);
+
 				// Apply acceleration that opposed velocity (basically 'drag')
-				ComponentLinearAcc += TransformWorldVectorToSimSpace(Output, -ComponentLinearVelocity) * ComponentLinearVelScale;
+				if(bScalingEnabled)
+				{
+					ComponentLinearAcc += TransformWorldVectorToSimSpaceScaled(Output, -ComponentLinearVelocity) * ComponentLinearVelScale;
+				}
+				else
+				{ 
+					ComponentLinearAcc += TransformWorldVectorToSimSpace(Output, -ComponentLinearVelocity) * ComponentLinearVelScale;
+				}
 
 				// Calc linear acceleration
 				const FVector ComponentLinearAcceleration = (ComponentLinearVelocity - PreviousComponentLinearVelocity) / NextTimeStep;
 				PreviousComponentLinearVelocity = ComponentLinearVelocity;
 				// Apply opposite acceleration to bodies
-				ComponentLinearAcc += TransformWorldVectorToSimSpace(Output, -ComponentLinearAcceleration) * ComponentLinearAccScale;
+				if (bScalingEnabled)
+				{
+					ComponentLinearAcc += TransformWorldVectorToSimSpaceScaled(Output, -ComponentLinearAcceleration) * ComponentLinearAccScale;
+				}
+				else
+				{
+					ComponentLinearAcc += TransformWorldVectorToSimSpace(Output, -ComponentLinearAcceleration) * ComponentLinearAccScale;
+				}
 
 				// Clamp ComponentLinearAcc to desired strength.	
 				FVector LinearAccClamp = ComponentAppliedLinearAccClamp;
@@ -356,7 +575,7 @@ void FAnimNode_AnimDynamics::EvaluateSkeletalControl_AnyThread(FComponentSpacePo
 					LinearAccClamp.Set(LinearAccClampOverride, LinearAccClampOverride, LinearAccClampOverride);
 				}
 
-				ComponentLinearAcc = ComponentLinearAcc.BoundToBox(-LinearAccClamp, LinearAccClamp);	
+				ComponentLinearAcc = ComponentLinearAcc.BoundToBox(-LinearAccClamp, LinearAccClamp);
 			}
 
 			// Update gravity.
@@ -398,7 +617,7 @@ void FAnimNode_AnimDynamics::EvaluateSkeletalControl_AnyThread(FComponentSpacePo
 				for (int32 Iter = 0; Iter < NumIters; ++Iter)
 				{
 					UpdateLimits(Output);
-					FAnimPhys::PhysicsUpdate(FixedTimeStep, SimBodies, LinearLimits, AngularLimits, Springs, SimSpaceGravityDirection, OrientedExternalForce, ComponentLinearAcc, NumSolverIterationsPreUpdate, NumSolverIterationsPostUpdate);
+					FAnimPhys::PhysicsUpdate(FixedTimeStep, SimBodies, LinearLimits, AngularLimits, Springs, SimSpaceGravityDirection, OrientedExternalForce, ComponentLinearAcc, SimFrameAngularAcc, SimFrameAngularVelocity, NumSolverIterationsPreUpdate, NumSolverIterationsPostUpdate);
 				}
 			}
 			else
@@ -409,9 +628,9 @@ void FAnimNode_AnimDynamics::EvaluateSkeletalControl_AnyThread(FComponentSpacePo
 				NextTimeStep = FMath::Min(NextTimeStep, MaxDeltaTime);
 
 				UpdateLimits(Output);
-				FAnimPhys::PhysicsUpdate(NextTimeStep, SimBodies, LinearLimits, AngularLimits, Springs, SimSpaceGravityDirection, OrientedExternalForce, ComponentLinearAcc, NumSolverIterationsPreUpdate, NumSolverIterationsPostUpdate);
+				FAnimPhys::PhysicsUpdate(NextTimeStep, SimBodies, LinearLimits, AngularLimits, Springs, SimSpaceGravityDirection, OrientedExternalForce, ComponentLinearAcc, SimFrameAngularAcc, SimFrameAngularVelocity, NumSolverIterationsPreUpdate, NumSolverIterationsPostUpdate);
 			}
-
+			
 #if ENABLE_ANIM_DRAW_DEBUG
 			DrawBodies(Output, SimBodies);
 #endif
@@ -440,7 +659,17 @@ void FAnimNode_AnimDynamics::EvaluateSkeletalControl_AnyThread(FComponentSpacePo
 				FCompactPoseBoneIndex BoneIndex = CurrentChainBone.GetCompactPoseIndex(BoneContainer);
 
 				// Calculate target bone transform from physics body.
-				FTransform NewBoneTransform(CurrentBody.Pose.Orientation, CurrentBody.Pose.Position - CurrentBody.Pose.Orientation.RotateVector(PhysicsBodyJointOffsets[Idx]));
+
+				FTransform NewBoneTransform;
+				
+				if (bScalingEnabled)
+				{
+					NewBoneTransform = FTransform(CurrentBody.Pose.Orientation, CurrentBody.Pose.Position - CurrentBody.Pose.Orientation.RotateVector(PhysicsBodyJointOffsets[Idx]), FVector(CurrentBody.Pose.Scale));
+				}
+				else
+				{
+					NewBoneTransform = FTransform(CurrentBody.Pose.Orientation, CurrentBody.Pose.Position - CurrentBody.Pose.Orientation.RotateVector(PhysicsBodyJointOffsets[Idx]));
+				}
 
 				if (RetargetingSettings.bEnabled)
 				{
@@ -694,6 +923,8 @@ void FAnimNode_AnimDynamics::RequestInitialise(ETeleportType InTeleportType)
 
 void FAnimNode_AnimDynamics::InitPhysics(FComponentSpacePoseContext& Output)
 {
+	const int bScalingSupportEnabled = CVarEnableScaling.GetValueOnAnyThread();
+
 	switch (InitTeleportType)
 	{
 	case ETeleportType::ResetPhysics:
@@ -728,16 +959,26 @@ void FAnimNode_AnimDynamics::InitPhysics(FComponentSpacePoseContext& Output)
 		for (FAnimPhysBodyDefinition& PhysicsBodyDef : PhysicsBodyDefinitions)
 		{
 			TArray<FAnimPhysShape> BodyShapes;
-			BodyShapes.Add(FAnimPhysShape::MakeBox(PhysicsBodyDef.BoxExtents));
 
 			PhysicsBodyDef.BoundBone.Initialize(BoneContainer);
 
 			FTransform BodyTransform = GetBoneTransformInSimSpace(Output, PhysicsBodyDef.BoundBone.GetCompactPoseIndex(BoneContainer));
-			BodyTransform.SetTranslation(BodyTransform.GetTranslation() + BodyTransform.GetRotation().RotateVector(PhysicsBodyDef.LocalJointOffset)); // Transform for physics body in Sim Space.
+
+			if (bScalingSupportEnabled)
+			{
+				BodyShapes.Add(FAnimPhysShape::MakeBox(PhysicsBodyDef.BoxExtents * BodyTransform.GetScale3D()));
+				BodyTransform.SetTranslation(BodyTransform.TransformPosition(PhysicsBodyDef.LocalJointOffset));
+			}
+			else
+			{
+				BodyShapes.Add(FAnimPhysShape::MakeBox(PhysicsBodyDef.BoxExtents));
+				BodyTransform.SetTranslation(BodyTransform.GetTranslation() + BodyTransform.GetRotation().RotateVector(PhysicsBodyDef.LocalJointOffset)); // Transform for physics body in Sim Space.
+			}
 
 			FAnimPhysLinkedBody NewChainBody(BodyShapes, BodyTransform.GetTranslation(), PhysicsBodyDef.BoundBone);
 			FAnimPhysRigidBody& PhysicsBody = NewChainBody.RigidBody.PhysBody;
 			PhysicsBody.Pose.Orientation = BodyTransform.GetRotation();
+			PhysicsBody.Pose.Scale = BodyTransform.GetScale3D().X;
 			PhysicsBody.PreviousOrientation = PhysicsBody.Pose.Orientation;
 			PhysicsBody.NextOrientation = PhysicsBody.Pose.Orientation;
 			PhysicsBody.CollisionType = PhysicsBodyDef.CollisionType;
@@ -745,13 +986,13 @@ void FAnimNode_AnimDynamics::InitPhysics(FComponentSpacePoseContext& Output)
 			switch (PhysicsBody.CollisionType)
 			{
 			case AnimPhysCollisionType::CustomSphere:
-				PhysicsBody.SphereCollisionRadius = static_cast<float>(PhysicsBodyDef.SphereCollisionRadius);
+				PhysicsBody.SphereCollisionRadius = static_cast<float>(bScalingSupportEnabled ? PhysicsBodyDef.SphereCollisionRadius * BodyTransform.GetScale3D().X : PhysicsBodyDef.SphereCollisionRadius);
 				break;
 			case AnimPhysCollisionType::InnerSphere:
-				PhysicsBody.SphereCollisionRadius = static_cast<float>(PhysicsBodyDef.BoxExtents.GetAbsMin() / 2.0);
+				PhysicsBody.SphereCollisionRadius = static_cast<float>(bScalingSupportEnabled ? (PhysicsBodyDef.BoxExtents.GetAbsMin() * BodyTransform.GetScale3D().X / 2.0) : PhysicsBodyDef.BoxExtents.GetAbsMin() / 2.0);
 				break;
 			case AnimPhysCollisionType::OuterSphere:
-				PhysicsBody.SphereCollisionRadius = static_cast<float>(PhysicsBodyDef.BoxExtents.GetAbsMax() / 2.0);
+				PhysicsBody.SphereCollisionRadius = static_cast<float>(bScalingSupportEnabled ? (PhysicsBodyDef.BoxExtents.GetAbsMax() * BodyTransform.GetScale3D().X / 2.0) : PhysicsBodyDef.BoxExtents.GetAbsMax() / 2.0);
 				break;
 			default:
 				break;
@@ -859,6 +1100,7 @@ void FAnimNode_AnimDynamics::InitPhysics(FComponentSpacePoseContext& Output)
 	InitTeleportType = ETeleportType::None;
 	PreviousCompWorldSpaceTM = Output.AnimInstanceProxy->GetComponentTransform();
 	PreviousActorWorldSpaceTM = Output.AnimInstanceProxy->GetActorTransform();
+	PreviousSimSpaceTM = GetSimToWorldSpaceTransform(SimulationSpace, Output);
 }
 
 void FAnimNode_AnimDynamics::TermPhysics()
@@ -1342,6 +1584,63 @@ FVector FAnimNode_AnimDynamics::TransformWorldVectorToSimSpace(FComponentSpacePo
 	return OutVec;
 }
 
+FVector FAnimNode_AnimDynamics::TransformWorldVectorToSimSpaceScaled(FComponentSpacePoseContext& Output, const FVector& InVec) const
+{
+	FVector OutVec = InVec;
+
+	switch (SimulationSpace)
+	{
+	case AnimPhysSimSpaceType::Component:
+	{
+		OutVec = Output.AnimInstanceProxy->GetComponentTransform().InverseTransformVector(OutVec);
+
+		break;
+	}
+
+	case AnimPhysSimSpaceType::Actor:
+	{
+		OutVec = Output.AnimInstanceProxy->GetActorTransform().InverseTransformVector(OutVec);
+
+		break;
+	}
+
+	case AnimPhysSimSpaceType::RootRelative:
+	{
+		const FBoneContainer& RequiredBones = Output.Pose.GetPose().GetBoneContainer();
+
+		FCompactPoseBoneIndex RootBoneCompactIndex(0);
+
+		FTransform RelativeBoneTransform = Output.Pose.GetComponentSpaceTransform(RootBoneCompactIndex);
+		RelativeBoneTransform = Output.AnimInstanceProxy->GetComponentTransform() * RelativeBoneTransform;
+		OutVec = RelativeBoneTransform.InverseTransformVector(OutVec);
+
+		break;
+	}
+
+	case AnimPhysSimSpaceType::BoneRelative:
+	{
+		const FBoneContainer& RequiredBones = Output.Pose.GetPose().GetBoneContainer();
+		if (RelativeSpaceBone.IsValidToEvaluate(RequiredBones))
+		{
+			FTransform RelativeBoneTransform = Output.Pose.GetComponentSpaceTransform(RelativeSpaceBone.GetCompactPoseIndex(RequiredBones));
+			RelativeBoneTransform = Output.AnimInstanceProxy->GetComponentTransform() * RelativeBoneTransform;
+			OutVec = RelativeBoneTransform.InverseTransformVector(OutVec);
+		}
+
+		break;
+	}
+	case AnimPhysSimSpaceType::World:
+	{
+		break;
+	}
+
+	default:
+		break;
+	}
+
+	return OutVec;
+}
+
 void FAnimNode_AnimDynamics::ConvertSimulationSpace(FComponentSpacePoseContext& Output, AnimPhysSimSpaceType From, AnimPhysSimSpaceType To) const
 {
 	for(FAnimPhysRigidBody* Body : BaseBodyPtrs)
diff --git a/Engine/Source/Runtime/AnimGraphRuntime/Public/BoneControllers/AnimNode_AnimDynamics.h b/Engine/Source/Runtime/AnimGraphRuntime/Public/BoneControllers/AnimNode_AnimDynamics.h
--- a/Engine/Source/Runtime/AnimGraphRuntime/Public/BoneControllers/AnimNode_AnimDynamics.h
+++ b/Engine/Source/Runtime/AnimGraphRuntime/Public/BoneControllers/AnimNode_AnimDynamics.h
@@ -294,6 +294,57 @@ struct FAnimConstraintOffsetPair
 	FVector Body1Offset;
 };
 
+/**
+ * Settings for the system which passes rotational motion of the simulation's space into the simulation. This allows the simulation to pass a 
+ * fraction of the world space motion onto the bodies which allows Bone-Space and Component-Space simulations to react to world-space 
+ * movement in a controllable way.
+ */
+template <> struct TIsPODType<FAnimPhysSimSpaceSettings> { enum { Value = true }; };
+
+USTRUCT(BlueprintType)
+struct FAnimPhysSimSpaceSettings
+{
+	GENERATED_BODY()
+
+	ANIMGRAPHRUNTIME_API FAnimPhysSimSpaceSettings();
+
+	~FAnimPhysSimSpaceSettings() = default;
+	FAnimPhysSimSpaceSettings(FAnimPhysSimSpaceSettings const&) = default;
+	FAnimPhysSimSpaceSettings& operator=(const FAnimPhysSimSpaceSettings&) = default;
+
+	/**
+	 * Global multipler on the effects of simulation space rotational movement. Must be in range[0, 1].If SimSpaceAngularAlpha = 0.0, the system is disabled and the simulation will
+	 * be fully local (i.e., world-space actor movement and rotation does not affect the simulation).
+	 */
+	UPROPERTY(EditAnywhere, BlueprintReadWrite, Category = Settings, meta = (ClampMin = "0.0", ClampMax = "1.0"))
+	float SimSpaceAngularAlpha;
+
+	/**
+	 * A clamp on the effective world-space angular velocity that is passed to the simulation. Units are radian/s, so a value of about 6.0 is one rotation per second.
+	 * The default value effectively means "unlimited". You would reduce this (and MaxAngularAcceleration) to limit how much bodies "fly out" when the actor spins on the spot.
+	 * This is especially useful if you have characters than can rotate very quickly and you would probably want values around or less than 10 in this case.
+	 */
+	UPROPERTY(EditAnywhere, BlueprintReadWrite, Category = Settings, meta = (ClampMin = "0.0"))
+	float MaxAngularVelocity;
+	
+	/**
+	 * A clamp on the effective world-space angular accleration that is passed to the simulation. Units are radian/s/s. The default value effectively means "unlimited".
+	 * This has a similar effect to MaxAngularVelocity, except that it is related to the flying out of bodies when the rotation speed suddenly changes. Typical limits for
+	 * a character might be around 100.
+	 */
+	UPROPERTY(EditAnywhere, BlueprintReadWrite, Category = Settings, meta = (ClampMin = "0.0"))
+	float MaxAngularAcceleration;
+
+	/**
+	 * Additional angular velocity that is added to the component angular velocity. This can be used to make the simulation act as if the actor is rotating
+	 * even when it is not. E.g., to apply physics to a character on a podium as the camera rotates around it, to emulate the podium itself rotating.
+	 * Vector is in world space. Units are rad/s.
+	 */
+	UPROPERTY(EditAnywhere, BlueprintReadWrite, Category = Settings)
+	FVector ExternalAngularVelocity;
+
+};
+
 
 USTRUCT(BlueprintInternalUseOnly)
 struct FAnimNode_AnimDynamics : public FAnimNode_SkeletalControlBase
@@ -317,6 +368,7 @@ struct FAnimNode_AnimDynamics : public FAnimNode_SkeletalControlBase
 	// Previous component & actor transforms, used to account for teleports
 	FTransform PreviousCompWorldSpaceTM;
 	FTransform PreviousActorWorldSpaceTM;
+	FTransform PreviousSimSpaceTM;
 
 	/** When in BoneRelative sim space, the simulation will use this bone as the origin */
 	UPROPERTY(EditAnywhere, Category = Setup, meta=(DisplayAfter="SimulationSpace", EditCondition = "SimulationSpace == AnimPhysSimSpaceType::BoneRelative"))
@@ -372,6 +424,19 @@ struct FAnimNode_AnimDynamics : public FAnimNode_SkeletalControlBase
 	UPROPERTY(EditAnywhere, Category = Settings)
 	FVector	ComponentAppliedLinearAccClamp;
 
+	/**
+	 * Settings for the system which passes motion of the simulation's space
+	 * into the simulation. This allows the simulation to pass a
+	 * fraction of the world space motion onto the bodies which allows Bone-Space
+	 * and Component-Space simulations to react to world-space movement in a
+	 * controllable way.
+	 * This system is a superset of the functionality provided by ComponentLinearAccScale,
+	 * ComponentLinearVelScale, and ComponentAppliedLinearAccClamp. In general
+	 * you should not have both systems enabled.
+	 */
+	UPROPERTY(EditAnywhere, Category = Settings, meta = (PinHiddenByDefault))
+	FAnimPhysSimSpaceSettings SimSpaceSettings;
+
 	/** Overridden angular bias value
 	 *  Angular bias is essentially a twist reduction for chain forces and defaults to a value to keep chains stability
 	*  in check. When using single-body systems sometimes angular forces will look like they are "catching-up" with
@@ -540,14 +605,21 @@ private:
 	ANIMGRAPHRUNTIME_API FTransform GetComponentSpaceTransformFromSimSpace(AnimPhysSimSpaceType SimSpace, FComponentSpacePoseContext& Output, const FTransform& InSimTransform, const FTransform& InCompWorldSpaceTM, const FTransform& InActorWorldSpaceTM) const;
 	ANIMGRAPHRUNTIME_API FTransform GetComponentSpaceTransformFromSimSpace(AnimPhysSimSpaceType SimSpace, const USkeletalMeshComponent* const SkelComp, const FTransform& InSimTransform) const;
 
+	FTransform GetSimToWorldSpaceTransform(AnimPhysSimSpaceType SimSpace, FComponentSpacePoseContext& Output) const;
+
 	// Given a transform in component space, convert it to the current sim space
 	ANIMGRAPHRUNTIME_API FTransform GetSimSpaceTransformFromComponentSpace(AnimPhysSimSpaceType SimSpace, FComponentSpacePoseContext& Output, const FTransform& InComponentTransform) const;
 
 	// Given a world-space vector, convert it into the current simulation space
 	ANIMGRAPHRUNTIME_API FVector TransformWorldVectorToSimSpace(FComponentSpacePoseContext& Output, const FVector& InVec) const;
 
+	ANIMGRAPHRUNTIME_API FVector TransformWorldVectorToSimSpaceScaled(FComponentSpacePoseContext& Output, const FVector& InVec) const;
+
 	ANIMGRAPHRUNTIME_API void ConvertSimulationSpace(FComponentSpacePoseContext& Output, AnimPhysSimSpaceType From, AnimPhysSimSpaceType To) const;
 
+	void CalculateSimulationSpaceMotion(AnimPhysSimSpaceType Space, FAnimPhysSimSpaceSettings Settings, FComponentSpacePoseContext& Output, const float Dt,
+		FVector& SpaceAngularVel, FVector& SpaceAngularAcc);
+
 	// Maximum time to consider when accumulating time debt to avoid spiraling
 	static ANIMGRAPHRUNTIME_API const float MaxTimeDebt;
 
@@ -599,6 +671,9 @@ private:
 	// Previous linear velocity to resolve world accelerations when not using world space simulation
 	FVector PreviousComponentLinearVelocity;
 
+	// Previous angular velocity to resolve world angular acceleration when not using world space simulation
+	FVector PreviousSimSpaceAngularVelocity;
+
 	//////////////////////////////////////////////////////////////////////////
 	// Live debug
 	//////////////////////////////////////////////////////////////////////////
diff --git a/Engine/Source/Runtime/Engine/Private/Animation/AnimPhysicsSolver.cpp b/Engine/Source/Runtime/Engine/Private/Animation/AnimPhysicsSolver.cpp
--- a/Engine/Source/Runtime/Engine/Private/Animation/AnimPhysicsSolver.cpp
+++ b/Engine/Source/Runtime/Engine/Private/Animation/AnimPhysicsSolver.cpp
@@ -46,6 +46,15 @@ DEFINE_STAT(STAT_AnimDynamicsLimitUpdate);
 bool FAnimPhys::bEnableDetailedStats = false;
 static FAutoConsoleVariableRef CVarAnimDynamicsDetailedStats(TEXT("p.AnimDynamicsDetailedStats"), FAnimPhys::bEnableDetailedStats, TEXT("When set to 1, will enable more detailed stats."));
 
+float FAnimPhys::CentrifugalAlpha = 1.0f;
+static FAutoConsoleVariableRef CVarAnimDynamicsSimSpaceCentrifugalAlpha(TEXT("p.AnimDynamics.SimSpaceCentrifugalAlpha"), FAnimPhys::CentrifugalAlpha, TEXT("Settings for centrifugal force alpha for anim dynamics nodes"));
+
+float FAnimPhys::CoriolisAlpha = 0.5f;
+static FAutoConsoleVariableRef CVarAnimDynamicsSimSpaceCoriolisAlpha(TEXT("p.AnimDynamics.SimSpaceCoriolisAlpha"), FAnimPhys::CoriolisAlpha, TEXT("Settings for centrifugal force alpha for anim dynamics nodes"));
+
+float FAnimPhys::EulerAlpha = 1.0f;
+static FAutoConsoleVariableRef CVarAnimDynamicsSimSpaceEulerAlpha(TEXT("p.AnimDynamics.SimSpaceEulerAlpha"), FAnimPhys::EulerAlpha, TEXT("Settings for euler force alpha for anim dynamics nodes"));
+
 //////////////////////////////////////////////////////////////////////////
 // FAnimPhysShape - Defines a set of vertices that make a shape with volume and CoM
 
@@ -180,7 +189,7 @@ FAnimPhysRigidBody::FAnimPhysRigidBody(TArray<FAnimPhysShape>& InShapes, const F
 		}
 	}
 
-	FMatrix InertiaTensor = FAnimPhys::CalculateInertia(Shapes, CenterOfMass);
+	InertiaTensor = FAnimPhys::CalculateInertia(Shapes, CenterOfMass);
 	InverseMass = 1.0f / Mass;
 	InverseTensorWithoutMass = InertiaTensor.InverseFast();
 	InverseWorldSpaceTensor = InverseTensorWithoutMass * InverseMass;
@@ -476,6 +485,7 @@ void FAnimPhys::ScaleRigidBodyMass(FAnimPhysRigidBody* InOutRigidBody, float Sca
 	InOutRigidBody->InverseMass *= 1.0f / Scale;
 	InOutRigidBody->AngularMomentum *= Scale;
 	InOutRigidBody->InverseWorldSpaceTensor *= 1.0f / Scale;
+	InOutRigidBody->InertiaTensor *= Scale;
 }
 
 FQuat FAnimPhys::DiffQ(const FQuat& InOrientation, const FMatrix &InInverseTensor, const FVector& InAngularMomentum)
@@ -909,7 +919,7 @@ void FAnimPhys::UpdatePose(FAnimPhysRigidBody* InBody)
 	InBody->InverseWorldSpaceTensor = OrientAsMatrix * InBody->InverseTensorWithoutMass * InBody->InverseMass * OrientAsMatrix.GetTransposed();
 }
 
-void FAnimPhys::PhysicsUpdate(float DeltaTime, TArray<FAnimPhysRigidBody*>& Bodies, TArray<FAnimPhysLinearLimit>& LinearLimits, TArray<FAnimPhysAngularLimit>& AngularLimits, TArray<FAnimPhysSpring>& Springs, const FVector& GravityDirection, const FVector& ExternalForce, const FVector& ExternalLinearAcc, int32 NumPreIterations, int32 NumPostIterations)
+void FAnimPhys::PhysicsUpdate(float DeltaTime, TArray<FAnimPhysRigidBody*>& Bodies, TArray<FAnimPhysLinearLimit>& LinearLimits, TArray<FAnimPhysAngularLimit>& AngularLimits, TArray<FAnimPhysSpring>& Springs, const FVector& GravityDirection, const FVector& ExternalForce, const FVector& ExternalLinearAcc, const FVector& ExternalAngularAcc, const FVector& ExternalAngularVelocity, int32 NumPreIterations, int32 NumPostIterations)
 {
 	CONDITIONAL_SCOPE_CYCLE_COUNTER(STAT_AnimDynamicsUpdate, bEnableDetailedStats);
 
@@ -927,7 +937,7 @@ void FAnimPhys::PhysicsUpdate(float DeltaTime, TArray<FAnimPhysRigidBody*>& Bodi
 		}
 	}
 
-	// apply external acceleration
+	// Apply external acceleration
 	if (!ExternalLinearAcc.IsNearlyZero())
 	{
 		for (FAnimPhysRigidBody* Body : Bodies)
@@ -939,6 +949,32 @@ void FAnimPhys::PhysicsUpdate(float DeltaTime, TArray<FAnimPhysRigidBody*>& Bodi
 		}
 	}
 
+	// Apply fictitious Forces
+	if(!(ExternalAngularVelocity.IsNearlyZero() && ExternalAngularAcc.IsNearlyZero()))
+	{
+		for (FAnimPhysRigidBody* Body : Bodies)
+		{
+			const FVector Velocity = Body->LinearMomentum * Body->InverseMass;
+			const FVector AngularVelocity = Body->Spin();
+
+			const FVector WxR = FVector::CrossProduct(ExternalAngularVelocity, Body->Pose.Position);
+			const FVector WxWxR = FVector::CrossProduct(ExternalAngularVelocity, WxR);
+			const FVector WxV = FVector::CrossProduct(ExternalAngularVelocity, Velocity);
+			const FVector AxR = FVector::CrossProduct(ExternalAngularAcc, Body->Pose.Position);
+
+			// Calculate the phantom forces
+			const FVector CoriolisAcc = CoriolisAlpha * 2.0f * WxV;
+			const FVector CentrifugalAcc = CentrifugalAlpha * WxWxR;
+			const FVector EulerAcc = EulerAlpha * AxR;
+			const FVector AngularAcc = AngularAccelerationAlpha * ExternalAngularAcc;
+
+			// Apply the phantom forces
+			Body->LinearMomentum -= (CoriolisAcc + CentrifugalAcc + EulerAcc) / Body->InverseMass * DeltaTime;
+			Body->AngularMomentum -= Body->InertiaTensor.TransformPosition(AngularAcc) * DeltaTime;
+		}
+	}
+	
+
 	for (FAnimPhysSpring& Spring : Springs)
 	{
 		Spring.ApplyForces(DeltaTime);
diff --git a/Engine/Source/Runtime/Engine/Public/Animation/AnimPhysicsSolver.h b/Engine/Source/Runtime/Engine/Public/Animation/AnimPhysicsSolver.h
--- a/Engine/Source/Runtime/Engine/Public/Animation/AnimPhysicsSolver.h
+++ b/Engine/Source/Runtime/Engine/Public/Animation/AnimPhysicsSolver.h
@@ -115,15 +115,24 @@ struct FAnimPhysPose
 {
 	FVector Position;
 	FQuat Orientation;
+	float Scale;
 
 	FAnimPhysPose(const FVector& InPosition, const FQuat& InOrient)
 		: Position(InPosition)
 		, Orientation(InOrient)
+		, Scale(1.0)
+	{}
+
+	FAnimPhysPose(const FVector& InPosition, const FQuat& InOrient, const float& InScale)
+		: Position(InPosition)
+		, Orientation(InOrient)
+		, Scale(InScale)
 	{}
 
 	FAnimPhysPose()
 		: Position(FVector::ZeroVector)
 		, Orientation(FQuat::Identity)
+		, Scale(1.0)
 	{
 
 	}
@@ -221,6 +230,8 @@ class FAnimPhysRigidBody : public FAnimPhysState
 	// Full (with mass) inverse tensor in world space
 	FMatrix				InverseWorldSpaceTensor;    // Inverse Inertia Tensor rotated into world space 
 
+	FMatrix				InertiaTensor;
+
 	// Maintained state, Start/Prev/Next
 	FVector				NextPosition;
 	FQuat				NextOrientation;
@@ -414,9 +425,21 @@ class FAnimPhys
 {
 public:
 
-	// Runtime cvar toggle for detailed/verbose profiling
+	// Runtime CVar toggle for detailed/verbose profiling
 	static ENGINE_API bool bEnableDetailedStats;
 
+	// How much of the coriolis force to apply to the sim bodies
+	static ENGINE_API float CoriolisAlpha;
+
+	// How much of the centrifugal force to apply to the sim bodies
+	static ENGINE_API float CentrifugalAlpha;
+
+	// How much of the euler force to apply to the sim bodies
+	static ENGINE_API float EulerAlpha;
+
+	// How much of the simulation frame's angular acceleration to pass onto the sim bodies.
+	static constexpr float AngularAccelerationAlpha = 1.0f;
+
 	/** Calculates the volume of a shape
 	 *  @param InVertices Verts in the shape
 	 *  @param InTriangles Triangles represented in InVertices
@@ -475,7 +498,7 @@ public:
 	*  @param NumPreIterations Number of times to iterate the limits before performing the integration
 	*  @param NumPostIterations Number of times to iterae the limits after performing the integration
 	*/
-	static ENGINE_API void PhysicsUpdate(float DeltaTime, TArray<FAnimPhysRigidBody*>& Bodies, TArray<FAnimPhysLinearLimit>& LinearLimits, TArray<FAnimPhysAngularLimit>& AngularLimits, TArray<FAnimPhysSpring>& Springs, const FVector& GravityDirection, const FVector& ExternalForce, const FVector& ExternalLinearAcc, int32 NumPreIterations = 8, int32 NumPostIterations = 2);
+	static ENGINE_API void PhysicsUpdate(float DeltaTime, TArray<FAnimPhysRigidBody*>& Bodies, TArray<FAnimPhysLinearLimit>& LinearLimits, TArray<FAnimPhysAngularLimit>& AngularLimits, TArray<FAnimPhysSpring>& Springs, const FVector& GravityDirection, const FVector& ExternalForce, const FVector& ExternalLinearAcc, const FVector& ExternalAngularAcc, const FVector& ExternalAngularVelocity, int32 NumPreIterations = 8, int32 NumPostIterations = 2);
 
 	//////////////////////////////////////////////////////////////////////////
 	// Constraint functions
