Ensure RBAN doesn't write NaNs into the pose, because if it does the editor can assert (and crash).

diff --git a/Engine/Source/Runtime/AnimGraphRuntime/Private/BoneControllers/AnimNode_RigidBody.cpp b/Engine/Source/Runtime/AnimGraphRuntime/Private/BoneControllers/AnimNode_RigidBody.cpp
--- a/Engine/Source/Runtime/AnimGraphRuntime/Private/BoneControllers/AnimNode_RigidBody.cpp
+++ b/Engine/Source/Runtime/AnimGraphRuntime/Private/BoneControllers/AnimNode_RigidBody.cpp
@@ -858,60 +858,62 @@ void FAnimNode_RigidBody::EvaluateSkeletalControl_AnyThread(FComponentSpacePoseC
 			if (BodyAnimData[BodyIndex].bIsSimulated)
 			{
 				FTransform BodyTM = Bodies[BodyIndex]->GetWorldTransform();
-
-				// if we clamp translation, we only do this when all linear translation are locked
-				// 
-				// @todo(ccaulfield): this shouldn't be required with Chaos - projection should be handling it...
-				if (bClampLinearTranslationLimitToRefPose
-					&&BodyAnimData[BodyIndex].LinearXMotion == ELinearConstraintMotion::LCM_Locked
-					&& BodyAnimData[BodyIndex].LinearYMotion == ELinearConstraintMotion::LCM_Locked
-					&& BodyAnimData[BodyIndex].LinearZMotion == ELinearConstraintMotion::LCM_Locked)
+				if (ensure(!BodyTM.ContainsNaN()))
 				{
-					// grab local space of length from ref pose 
-					// we have linear limit value - see if that works
-					// calculate current local space from parent
-					// find parent transform
-					const int32 ParentBodyIndex = OutputData.ParentBodyIndex;
-					FTransform ParentTransform = FTransform::Identity;
-					if (ParentBodyIndex != INDEX_NONE)
+					// if we clamp translation, we only do this when all linear translation are locked
+					// 
+					// @todo(ccaulfield): this shouldn't be required with Chaos - projection should be handling it...
+					if (bClampLinearTranslationLimitToRefPose
+						&&BodyAnimData[BodyIndex].LinearXMotion == ELinearConstraintMotion::LCM_Locked
+						&& BodyAnimData[BodyIndex].LinearYMotion == ELinearConstraintMotion::LCM_Locked
+						&& BodyAnimData[BodyIndex].LinearZMotion == ELinearConstraintMotion::LCM_Locked)
 					{
-						ParentTransform = Bodies[ParentBodyIndex]->GetWorldTransform();
-					}
+						// grab local space of length from ref pose 
+						// we have linear limit value - see if that works
+						// calculate current local space from parent
+						// find parent transform
+						const int32 ParentBodyIndex = OutputData.ParentBodyIndex;
+						FTransform ParentTransform = FTransform::Identity;
+						if (ParentBodyIndex != INDEX_NONE)
+						{
+							ParentTransform = Bodies[ParentBodyIndex]->GetWorldTransform();
+						}
 
-					// get local transform
-					FTransform LocalTransform = BodyTM.GetRelativeTransform(ParentTransform);
-					const float CurrentLength = static_cast<float>(LocalTransform.GetTranslation().Size());
+						// get local transform
+						FTransform LocalTransform = BodyTM.GetRelativeTransform(ParentTransform);
+						const float CurrentLength = static_cast<float>(LocalTransform.GetTranslation().Size());
 
-					// this is inconsistent with constraint. The actual linear limit is set by constraint
-					if (!FMath::IsNearlyEqual(CurrentLength, BodyAnimData[BodyIndex].RefPoseLength, KINDA_SMALL_NUMBER))
-					{
-						float RefPoseLength = BodyAnimData[BodyIndex].RefPoseLength;
-						if (CurrentLength > RefPoseLength)
+						// this is inconsistent with constraint. The actual linear limit is set by constraint
+						if (!FMath::IsNearlyEqual(CurrentLength, BodyAnimData[BodyIndex].RefPoseLength, KINDA_SMALL_NUMBER))
 						{
-							float Scale = (CurrentLength > KINDA_SMALL_NUMBER) ? RefPoseLength / CurrentLength : 0.f;
-							// we don't use 1.f here because 1.f can create pops based on float issue. 
-							// so we only activate clamping when less than 90%
-							if (Scale < 0.9f)
+							float RefPoseLength = BodyAnimData[BodyIndex].RefPoseLength;
+							if (CurrentLength > RefPoseLength)
 							{
-								LocalTransform.ScaleTranslation(Scale);
-								BodyTM = LocalTransform * ParentTransform;
-								Bodies[BodyIndex]->SetWorldTransform(BodyTM);
+								float Scale = (CurrentLength > KINDA_SMALL_NUMBER) ? RefPoseLength / CurrentLength : 0.f;
+								// we don't use 1.f here because 1.f can create pops based on float issue. 
+								// so we only activate clamping when less than 90%
+								if (Scale < 0.9f)
+								{
+									LocalTransform.ScaleTranslation(Scale);
+									BodyTM = LocalTransform * ParentTransform;
+									Bodies[BodyIndex]->SetWorldTransform(BodyTM);
+								}
 							}
 						}
 					}
-				}
 
-				FTransform ComponentSpaceTM;
+					FTransform ComponentSpaceTM;
 
-				switch(SimulationSpace)
-				{
-					case ESimulationSpace::ComponentSpace: ComponentSpaceTM = BodyTM; break;
-					case ESimulationSpace::WorldSpace: ComponentSpaceTM = BodyTM.GetRelativeTransform(SimulationWorldSpaceTM); break;
-					case ESimulationSpace::BaseBoneSpace: ComponentSpaceTM = BodyTM * BaseBoneTM; break;
-					default: ensureMsgf(false, TEXT("Unsupported Simulation Space")); ComponentSpaceTM = BodyTM;
-				}
+					switch(SimulationSpace)
+					{
+						case ESimulationSpace::ComponentSpace: ComponentSpaceTM = BodyTM; break;
+						case ESimulationSpace::WorldSpace: ComponentSpaceTM = BodyTM.GetRelativeTransform(SimulationWorldSpaceTM); break;
+						case ESimulationSpace::BaseBoneSpace: ComponentSpaceTM = BodyTM * BaseBoneTM; break;
+						default: ensureMsgf(false, TEXT("Unsupported Simulation Space")); ComponentSpaceTM = BodyTM;
+					}
 					
-				OutBoneTransforms.Add(FBoneTransform(OutputData.CompactPoseBoneIndex, ComponentSpaceTM));
+					OutBoneTransforms.Add(FBoneTransform(OutputData.CompactPoseBoneIndex, ComponentSpaceTM));
+				}
 			}
 		}
 
