Add option to treat default as simulated for RBAN.

diff --git a/Engine/Source/Runtime/AnimGraphRuntime/Private/BoneControllers/AnimNode_RigidBody.cpp b/Engine/Source/Runtime/AnimGraphRuntime/Private/BoneControllers/AnimNode_RigidBody.cpp
--- a/Engine/Source/Runtime/AnimGraphRuntime/Private/BoneControllers/AnimNode_RigidBody.cpp
+++ b/Engine/Source/Runtime/AnimGraphRuntime/Private/BoneControllers/AnimNode_RigidBody.cpp
@@ -1141,7 +1141,7 @@ void FAnimNode_RigidBody::InitPhysics(const UAnimInstance* InAnimInstance)
 			{
 				UBodySetup* BodySetup = UsePhysicsAsset->SkeletalBodySetups[BodyInstance->InstanceBodyIndex];
 
-				bool bSimulated = (BodySetup->PhysicsType == EPhysicsType::PhysType_Simulated);
+				bool bSimulated = (BodySetup->PhysicsType == EPhysicsType::PhysType_Simulated) || (bUseDefaultAsSimulated && BodySetup->PhysicsType == EPhysicsType::PhysType_Default);
 				ImmediatePhysics::EActorType ActorType = bSimulated ?  ImmediatePhysics::EActorType::DynamicActor : ImmediatePhysics::EActorType::KinematicActor;
 				ImmediatePhysics::FActorHandle* NewBodyHandle = PhysicsSimulation->CreateActor(ImmediatePhysics::MakeActorSetup(ActorType, BodyInstance, BodyInstance->GetUnrealWorldTransform()));
 				if (NewBodyHandle)
diff --git a/Engine/Source/Runtime/AnimGraphRuntime/Public/BoneControllers/AnimNode_RigidBody.h b/Engine/Source/Runtime/AnimGraphRuntime/Public/BoneControllers/AnimNode_RigidBody.h
--- a/Engine/Source/Runtime/AnimGraphRuntime/Public/BoneControllers/AnimNode_RigidBody.h
+++ b/Engine/Source/Runtime/AnimGraphRuntime/Public/BoneControllers/AnimNode_RigidBody.h
@@ -206,6 +206,10 @@ public:
 	UPROPERTY(EditAnywhere, Category = Settings)
 	bool bDefaultToSkeletalMeshPhysicsAsset = true;
 
+	/** Treat parts in the Physics Asset with PhysicsType Default as Simulated for RBAN. */
+	UPROPERTY(EditAnywhere, Category = Settings)
+	bool bUseDefaultAsSimulated = false;
+
 private:
 	/** Get the physics asset candidate to be used while respecting the bDefaultToSkeletalMeshPhysicsAsset and the priority to the override physics asset. */
 	UPhysicsAsset* GetPhysicsAssetToBeUsed(const UAnimInstance* InAnimInstance) const;
