When Anim Node was colliding with world objects just after starting the simulation, the simulation blew up.
Collect world object ignoring the simulation space transform, cannot update the transform properly.
Used a lazy initilaztion to fix the issue.

diff --git a/Engine/Source/Runtime/AnimGraphRuntime/Private/BoneControllers/AnimNode_RigidBody.cpp b/Engine/Source/Runtime/AnimGraphRuntime/Private/BoneControllers/AnimNode_RigidBody.cpp
--- a/Engine/Source/Runtime/AnimGraphRuntime/Private/BoneControllers/AnimNode_RigidBody.cpp
+++ b/Engine/Source/Runtime/AnimGraphRuntime/Private/BoneControllers/AnimNode_RigidBody.cpp
@@ -1780,15 +1780,16 @@ void FAnimNode_RigidBody::UpdateWorldObjects(const FTransform& SpaceTransform)
 
 	if (SimulationSpace != ESimulationSpace::WorldSpace)
 	{
-		for (const auto& WorldEntry : ComponentsInSim)
+		for (TPair<const UPrimitiveComponent*, FWorldObject>& WorldEntry : ComponentsInSim)
 		{ 
 			const UPrimitiveComponent* OverlapComp = WorldEntry.Key;
 			if (OverlapComp != nullptr)
 			{
-				ImmediatePhysics::FActorHandle* ActorHandle = WorldEntry.Value.ActorHandle;
+				FWorldObject& WorldObject = WorldEntry.Value;
+				ImmediatePhysics::FActorHandle* ActorHandle = WorldObject.ActorHandle;
 
 				// Calculate the sim-space transform of this object
-				const FTransform CompWorldTransform = OverlapComp->BodyInstance.GetUnrealWorldTransform();
+				const FTransform CompWorldTransform = OverlapComp->GetComponentTransform();
 				FTransform CompSpaceTransform;
 				CompSpaceTransform.SetTranslation(SpaceTransform.InverseTransformPosition(CompWorldTransform.GetLocation()));
 				CompSpaceTransform.SetRotation(SpaceTransform.InverseTransformRotation(CompWorldTransform.GetRotation()));
@@ -1796,6 +1797,14 @@ void FAnimNode_RigidBody::UpdateWorldObjects(const FTransform& SpaceTransform)
 
 				// Update the sim's copy of the world object
 				ActorHandle->SetKinematicTarget(CompSpaceTransform);
+
+				// We need to update the particle's transform in the right space for the first time. 
+				// When actor is created in CollectWorldObjects the space transform is still unknown.
+				if (WorldObject.bNew)
+				{
+					WorldObject.bNew = false;
+					ActorHandle->GetParticle()->InitTransform(CompSpaceTransform.GetTranslation(), CompSpaceTransform.GetRotation());
+				}
 			}
 		}
 	}
diff --git a/Engine/Source/Runtime/AnimGraphRuntime/Public/BoneControllers/AnimNode_RigidBody.h b/Engine/Source/Runtime/AnimGraphRuntime/Public/BoneControllers/AnimNode_RigidBody.h
--- a/Engine/Source/Runtime/AnimGraphRuntime/Public/BoneControllers/AnimNode_RigidBody.h
+++ b/Engine/Source/Runtime/AnimGraphRuntime/Public/BoneControllers/AnimNode_RigidBody.h
@@ -472,12 +472,13 @@ private:
 
 	struct FWorldObject
 	{
-		FWorldObject() : ActorHandle(nullptr), LastSeenTick(0), bExpired(false) {}
+		FWorldObject() : ActorHandle(nullptr), LastSeenTick(0), bExpired(false), bNew(true) {}
 		FWorldObject(ImmediatePhysics::FActorHandle* InActorHandle, int32 InLastSeenTick) : ActorHandle(InActorHandle), LastSeenTick(InLastSeenTick), bExpired(false) {}
 
 		ImmediatePhysics::FActorHandle* ActorHandle;
 		int32 LastSeenTick;
-		bool bExpired;
+		uint8 bExpired : 1;
+		uint8 bNew : 1;
 	};
 
 	TArray<FOutputBoneData> OutputBoneData;
