Fix a non initialzied  variable in constructor

diff --git a/Engine/Source/Runtime/AnimGraphRuntime/Public/BoneControllers/AnimNode_RigidBody.h b/Engine/Source/Runtime/AnimGraphRuntime/Public/BoneControllers/AnimNode_RigidBody.h
--- a/Engine/Source/Runtime/AnimGraphRuntime/Public/BoneControllers/AnimNode_RigidBody.h
+++ b/Engine/Source/Runtime/AnimGraphRuntime/Public/BoneControllers/AnimNode_RigidBody.h
@@ -473,7 +473,7 @@ private:
 	struct FWorldObject
 	{
 		FWorldObject() : ActorHandle(nullptr), LastSeenTick(0), bExpired(false), bNew(true) {}
-		FWorldObject(ImmediatePhysics::FActorHandle* InActorHandle, int32 InLastSeenTick) : ActorHandle(InActorHandle), LastSeenTick(InLastSeenTick), bExpired(false) {}
+		FWorldObject(ImmediatePhysics::FActorHandle* InActorHandle, int32 InLastSeenTick) : ActorHandle(InActorHandle), LastSeenTick(InLastSeenTick), bExpired(false), bNew(true) {}
 
 		ImmediatePhysics::FActorHandle* ActorHandle;
 		int32 LastSeenTick;
