[BeatSyncedAnim] Add multi-beat alignment to BeatSyncedMetaData.
Rewrite of AnimNode_PlayBeatSyncedAnim removing most of its state, including the dreaded "sideloaded anim" which is now just represented by a form of alignment.

#tests PIE jam stage, frontend, and mainstage
#changelist validated
#virtualized

diff --git a/Engine/Source/Runtime/AnimGraphRuntime/Private/AnimNodes/AnimNode_SequenceEvaluator.cpp b/Engine/Source/Runtime/AnimGraphRuntime/Private/AnimNodes/AnimNode_SequenceEvaluator.cpp
--- a/Engine/Source/Runtime/AnimGraphRuntime/Private/AnimNodes/AnimNode_SequenceEvaluator.cpp
+++ b/Engine/Source/Runtime/AnimGraphRuntime/Private/AnimNodes/AnimNode_SequenceEvaluator.cpp
@@ -33,7 +33,12 @@ void FAnimNode_SequenceEvaluatorBase::CacheBones_AnyThread(const FAnimationCache
 void FAnimNode_SequenceEvaluatorBase::UpdateAssetPlayer(const FAnimationUpdateContext& Context)
 {
 	GetEvaluateGraphExposedInputs().Execute(Context);
+	// We split this off so subclasses can insert code between the two to set the explicit time.
+	PostExposedInputsUpdateAssetPlayer(Context);
+}
 
+void FAnimNode_SequenceEvaluatorBase::PostExposedInputsUpdateAssetPlayer(const FAnimationUpdateContext& Context)
+{
 	float CurrentExplicitTime = GetExplicitTime();
 
 	UAnimSequenceBase* CurrentSequence = GetSequence();
diff --git a/Engine/Source/Runtime/AnimGraphRuntime/Public/AnimNodes/AnimNode_SequenceEvaluator.h b/Engine/Source/Runtime/AnimGraphRuntime/Public/AnimNodes/AnimNode_SequenceEvaluator.h
--- a/Engine/Source/Runtime/AnimGraphRuntime/Public/AnimNodes/AnimNode_SequenceEvaluator.h
+++ b/Engine/Source/Runtime/AnimGraphRuntime/Public/AnimNodes/AnimNode_SequenceEvaluator.h
@@ -93,6 +93,10 @@ public:
 
 	// The start up position, it only applies when ReinitializationBehavior == StartPosition. Only used when bTeleportToExplicitTime is false.
 	virtual float GetStartPosition() const { return 0.0f; }
+
+	// Called after executing the input pins, split for subclass convenience.
+	ANIMGRAPHRUNTIME_API void PostExposedInputsUpdateAssetPlayer(const FAnimationUpdateContext& Context);
+
 };
 
 // Sequence evaluator node that can be used with constant folding
