Add a cVar in Anim Dynamic to disable the World Space transform
We can now force all Anim Dynamic to use the component space by default

diff --git a/Engine/Source/Runtime/AnimGraphRuntime/Private/BoneControllers/AnimNode_AnimDynamics.cpp b/Engine/Source/Runtime/AnimGraphRuntime/Private/BoneControllers/AnimNode_AnimDynamics.cpp
--- a/Engine/Source/Runtime/AnimGraphRuntime/Private/BoneControllers/AnimNode_AnimDynamics.cpp
+++ b/Engine/Source/Runtime/AnimGraphRuntime/Private/BoneControllers/AnimNode_AnimDynamics.cpp
@@ -28,6 +28,7 @@ TAutoConsoleVariable<int32> CVarAdaptiveSubstepNumDebtFrames(TEXT("p.AnimDynamic
 TAutoConsoleVariable<int32> CVarEnableWind(TEXT("p.AnimDynamicsWind"), 1, TEXT("Enables/Disables anim dynamics wind forces globally."), ECVF_Scalability);
 TAutoConsoleVariable<float> CVarComponentAppliedLinearAccClampOverride(TEXT("p.AnimDynamics.ComponentAppliedLinearAccClampOverride"), -1.0f, TEXT("Override the per asset setting for all axis (X,Y & Z) of ComponentAppliedLinearAccClamp for all Anim Dynamics Nodes. Negative values are ignored."));
 TAutoConsoleVariable<float> CVarGravityScale(TEXT("p.AnimDynamics.GravityScale"), 1.0f, TEXT("Multiplies the default gravity and the gravity override on all Anim Dynamics Nodes."));
+TAutoConsoleVariable<bool>  CVarDisableWorldSpace(TEXT("p.AnimDynamics.DisableWorldSpace"), false, TEXT("Disable World Sapce transform and override it to component space on all Anim Dynamics Nodes."));
 
 // FSimSpaceSettingsWithControl forced overrides for testing
 bool bAnimDynamics_SimSpace_EnableOverride = false;
@@ -130,7 +131,7 @@ void FAnimNode_AnimDynamics::DrawBodies(FComponentSpacePoseContext& InContext, c
 		}
 	}
 
-	if(SimulationSpace != AnimPhysSimSpaceType::World)
+	if(SimulationSpace != AnimPhysSimSpaceType::World || CVarDisableWorldSpace.GetValueOnAnyThread())
 	{
 		FTransform Origin;
 
@@ -285,6 +286,11 @@ FTransform FAnimNode_AnimDynamics::GetSimToWorldSpaceTransform(AnimPhysSimSpaceT
 {
 	FTransform OutTransform;
 
+	if (CVarDisableWorldSpace.GetValueOnAnyThread() && SimSpace == AnimPhysSimSpaceType::World)
+	{
+		SimSpace = AnimPhysSimSpaceType::Component;
+	}
+
 	switch (SimSpace)
 	{
 	case AnimPhysSimSpaceType::Component:
@@ -356,7 +362,7 @@ void FAnimNode_AnimDynamics::CalculateSimulationSpaceMotion(
 	}
 
 	// If the simulation is in world space, the simulation space is stationary
-	if (Space == AnimPhysSimSpaceType::World)
+	if (Space == AnimPhysSimSpaceType::World && !CVarDisableWorldSpace.GetValueOnAnyThread())
 	{
 		SpaceAngularVel = SimSpaceSettings.ExternalAngularVelocity;
 		return;
@@ -529,7 +535,7 @@ void FAnimNode_AnimDynamics::EvaluateSkeletalControl_AnyThread(FComponentSpacePo
 
 			CalculateSimulationSpaceMotion(SimulationSpace, SimSpaceSettings, Output, NextTimeStep, SimFrameAngularVelocity, SimFrameAngularAcc);
 
-			if (SimulationSpace != AnimPhysSimSpaceType::World)
+			if (SimulationSpace != AnimPhysSimSpaceType::World || CVarDisableWorldSpace.GetValueOnAnyThread())
 			{
 				// Transform Gravity Override into simulation space
 				if (bUseGravityOverride && !bGravityOverrideInSimSpace)
@@ -1395,6 +1401,11 @@ FTransform FAnimNode_AnimDynamics::GetComponentSpaceTransformFromSimSpace(AnimPh
 {
 	FTransform OutTransform = InSimTransform;
 
+	if (SimSpace == AnimPhysSimSpaceType::World && CVarDisableWorldSpace.GetValueOnAnyThread())
+	{
+		SimSpace = AnimPhysSimSpaceType::Component;
+	}
+
 	switch(SimSpace)
 	{
 		// Change nothing, already in component space
@@ -1473,6 +1484,11 @@ FTransform FAnimNode_AnimDynamics::GetSimSpaceTransformFromComponentSpace(AnimPh
 {
 	FTransform ResultTransform = InComponentTransform;
 
+	if (SimSpace == AnimPhysSimSpaceType::World && CVarDisableWorldSpace.GetValueOnAnyThread())
+	{
+		SimSpace = AnimPhysSimSpaceType::Component;
+	}
+
 	switch(SimSpace)
 	{
 		// Change nothing, already in component space
@@ -1531,6 +1547,12 @@ FVector FAnimNode_AnimDynamics::TransformWorldVectorToSimSpace(FComponentSpacePo
 {
 	FVector OutVec = InVec;
 
+	// If Disable World Space cVar is enabled act like it is in component space
+	if (SimulationSpace == AnimPhysSimSpaceType::World && CVarDisableWorldSpace.GetValueOnAnyThread())
+	{
+		OutVec = Output.AnimInstanceProxy->GetComponentTransform().InverseTransformVectorNoScale(OutVec);
+	}
+
 	switch(SimulationSpace)
 	{
 	case AnimPhysSimSpaceType::Component:
@@ -1588,6 +1610,12 @@ FVector FAnimNode_AnimDynamics::TransformWorldVectorToSimSpaceScaled(FComponentS
 {
 	FVector OutVec = InVec;
 
+	// If Disable World Space cVar is enabled act like it is in component space
+	if (SimulationSpace == AnimPhysSimSpaceType::World && CVarDisableWorldSpace.GetValueOnAnyThread())
+	{
+		OutVec = Output.AnimInstanceProxy->GetComponentTransform().InverseTransformVector(OutVec);
+	}
+
 	switch (SimulationSpace)
 	{
 	case AnimPhysSimSpaceType::Component:
