Fix for exploding pose with null aim offset input variable

diff --git a/Engine/Source/Runtime/AnimGraphRuntime/Private/AnimNodes/AnimNode_RotationOffsetBlendSpace.cpp b/Engine/Source/Runtime/AnimGraphRuntime/Private/AnimNodes/AnimNode_RotationOffsetBlendSpace.cpp
--- a/Engine/Source/Runtime/AnimGraphRuntime/Private/AnimNodes/AnimNode_RotationOffsetBlendSpace.cpp
+++ b/Engine/Source/Runtime/AnimGraphRuntime/Private/AnimNodes/AnimNode_RotationOffsetBlendSpace.cpp
@@ -81,7 +81,7 @@ void FAnimNode_RotationOffsetBlendSpace::Evaluate_AnyThread(FPoseContext& Contex
 	// Evaluate base pose
 	BasePose.Evaluate(Context);
 
-	if (bIsLODEnabled && FAnimWeight::IsRelevant(ActualAlpha))
+	if (bIsLODEnabled && FAnimWeight::IsRelevant(ActualAlpha) && GetBlendSpace())
 	{
 		// Evaluate MeshSpaceRotation additive blendspace
 		FPoseContext MeshSpaceRotationAdditivePoseContext(Context);
