Restrict AnimDynamics body scaling to World Space and add minimum scaling limits to avoid a simulation blowup

diff --git a/Engine/Source/Runtime/AnimGraphRuntime/Private/BoneControllers/AnimNode_AnimDynamics.cpp b/Engine/Source/Runtime/AnimGraphRuntime/Private/BoneControllers/AnimNode_AnimDynamics.cpp
--- a/Engine/Source/Runtime/AnimGraphRuntime/Private/BoneControllers/AnimNode_AnimDynamics.cpp
+++ b/Engine/Source/Runtime/AnimGraphRuntime/Private/BoneControllers/AnimNode_AnimDynamics.cpp
@@ -112,7 +112,7 @@ void FAnimNode_AnimDynamics::DrawBodies(FComponentSpacePoseContext& InContext, c
 			continue;
 		}
 
-		FTransform Transform(Body.Pose.Orientation, Body.Pose.Position, FVector(Body.Pose.Scale));
+		FTransform Transform(Body.Pose.Orientation, Body.Pose.Position, SimulationSpace != AnimPhysSimSpaceType::World || CVarDisableWorldSpace.GetValueOnAnyThread() ? FVector(Body.Pose.Scale) : FVector(1));
 		Transform = GetComponentSpaceTransformFromSimSpace(SimulationSpace, InContext, Transform);
 
 		Proxy->AnimDrawDebugCoordinateSystem(Transform.GetTranslation(), Transform.Rotator(), 2.0f, false, -1.0f, 0.15f);
@@ -191,6 +191,7 @@ FAnimNode_AnimDynamics::FAnimNode_AnimDynamics()
 , GravityOverride(ForceInitToZero)
 , LinearSpringConstant(0.0f)
 , AngularSpringConstant(0.0f)
+, WorldSpaceMinimumScale(0.01f)
 , WindScale(1.0f)
 , AngularBiasOverride(0.0f)
 , NumSolverIterationsPreUpdate(4)
@@ -391,6 +392,42 @@ void FAnimNode_AnimDynamics::CalculateSimulationSpaceMotion(
 
 }
 
+void FAnimNode_AnimDynamics::UpdateBodyScaling(const int32 BodyIndex, float Scale) const
+{
+	if (!PhysicsBodyDefinitions.IsValidIndex(BodyIndex))
+	{
+		return;
+	}
+
+	const FAnimPhysBodyDefinition& PhysicsBodyDef = PhysicsBodyDefinitions[BodyIndex];
+	FAnimPhysRigidBody* RigidBody = BaseBodyPtrs[BodyIndex];
+
+	if (RigidBody->Pose.Scale != Scale)
+	{
+		FAnimPhys::ScaleRigidBodyMass(RigidBody, Scale / RigidBody->Pose.Scale);
+
+		TArray<FAnimPhysShape> BodyShapes;
+
+		BodyShapes.Add(FAnimPhysShape::MakeBox(PhysicsBodyDef.BoxExtents * Scale));
+
+		switch (RigidBody->CollisionType)
+		{
+		case AnimPhysCollisionType::CustomSphere:
+			RigidBody->SphereCollisionRadius = static_cast<float>(PhysicsBodyDef.SphereCollisionRadius * Scale);
+			break;
+		case AnimPhysCollisionType::InnerSphere:
+			RigidBody->SphereCollisionRadius = static_cast<float>(PhysicsBodyDef.BoxExtents.GetAbsMin() * Scale / 2.0);
+			break;
+		case AnimPhysCollisionType::OuterSphere:
+			RigidBody->SphereCollisionRadius = static_cast<float>(PhysicsBodyDef.BoxExtents.GetAbsMax() * Scale / 2.0);
+			break;
+		default:
+			break;
+		}
+		RigidBody->Shapes = BodyShapes;
+	}
+}
+
 void FAnimNode_AnimDynamics::EvaluateSkeletalControl_AnyThread(FComponentSpacePoseContext& Output, TArray<FBoneTransform>& OutBoneTransforms)
 {
 	DECLARE_SCOPE_HIERARCHICAL_COUNTER_ANIMNODE(EvaluateSkeletalControl_AnyThread)
@@ -416,11 +453,11 @@ void FAnimNode_AnimDynamics::EvaluateSkeletalControl_AnyThread(FComponentSpacePo
 			InitTeleportType = ETeleportType::None;
 		}
 
+		const FTransform SimWorldTransform = GetSimToWorldSpaceTransform(SimulationSpace, Output);
+
 		if (bScalingEnabled)
 		{
-			const FTransform SimWorldTransform = GetSimToWorldSpaceTransform(SimulationSpace, Output);
-
-			if(SimWorldTransform.GetScale3D() != PreviousSimSpaceTM.GetScale3D())
+			if (!FMath::IsNearlyEqual(SimWorldTransform.GetScale3D().X, PreviousSimSpaceTM.GetScale3D().X))
 			{
 				// If the scaling of the sim space has changed, update the bodies to the proper scale
 				const FBoneContainer& BoneContainer = Output.Pose.GetPose().GetBoneContainer();
@@ -441,31 +478,16 @@ void FAnimNode_AnimDynamics::EvaluateSkeletalControl_AnyThread(FComponentSpacePo
 					FCompactPoseBoneIndex BoneIndex = CurrentBoneRef.GetCompactPoseIndex(BoneContainer);
 					FTransform BoundBoneTransform = GetBoneTransformInSimSpace(Output, BoneIndex);
 
-					float Scale = BoundBoneTransform.GetScale3D().X;
+					float Scale = BoundBoneTransform.GetScale3D().GetAbsMin();
+					Scale = Scale < UE_KINDA_SMALL_NUMBER ? UE_KINDA_SMALL_NUMBER : Scale;
+					
 					if (RigidBody.Pose.Scale != Scale)
 					{
-						FAnimPhys::ScaleRigidBodyMass(&RigidBody, Scale/RigidBody.Pose.Scale);
-
-						TArray<FAnimPhysShape> BodyShapes;
-
-						BodyShapes.Add(FAnimPhysShape::MakeBox(PhysicsBodyDef.BoxExtents * Scale));
-						RigidBody.Pose.Scale = Scale;
-
-						switch (RigidBody.CollisionType)
+						if (SimulationSpace == AnimPhysSimSpaceType::World && !CVarDisableWorldSpace.GetValueOnAnyThread())
 						{
-						case AnimPhysCollisionType::CustomSphere:
-							RigidBody.SphereCollisionRadius = static_cast<float>(PhysicsBodyDef.SphereCollisionRadius * Scale);
-							break;
-						case AnimPhysCollisionType::InnerSphere:
-							RigidBody.SphereCollisionRadius = static_cast<float>(PhysicsBodyDef.BoxExtents.GetAbsMin() * Scale / 2.0);
-							break;
-						case AnimPhysCollisionType::OuterSphere:
-							RigidBody.SphereCollisionRadius = static_cast<float>(PhysicsBodyDef.BoxExtents.GetAbsMax() * Scale / 2.0);
-							break;
-						default:
-							break;
+							UpdateBodyScaling(ActiveIndex, Scale);
 						}
-						RigidBody.Shapes = BodyShapes;
+						RigidBody.Pose.Scale = Scale;
 					}
 				}
 			}
@@ -482,12 +504,22 @@ void FAnimNode_AnimDynamics::EvaluateSkeletalControl_AnyThread(FComponentSpacePo
 
 				PhysBody.Pose.Position = BoneTransform.GetTranslation();
 				PhysBody.Pose.Orientation = BoneTransform.GetRotation();
-				PhysBody.Pose.Scale = BoneTransform.GetScale3D().X;
+
+				float Scale = BoneTransform.GetScale3D().GetAbsMin();
+				Scale = Scale < UE_KINDA_SMALL_NUMBER ? UE_KINDA_SMALL_NUMBER : Scale;
+				PhysBody.Pose.Scale = Scale;
+
 				PhysBody.LinearMomentum = FVector::ZeroVector;
 				PhysBody.AngularMomentum = FVector::ZeroVector;
 			}
 		}
 
+		// Disable simulation below minimum scale in world space mode.
+		if (bScalingEnabled && (SimulationSpace == AnimPhysSimSpaceType::World) && !CVarDisableWorldSpace.GetValueOnAnyThread() && (SimWorldTransform.GetScale3D().SizeSquared() < WorldSpaceMinimumScale * WorldSpaceMinimumScale))
+		{
+			return;
+		}
+
 		if (ShouldDoPhysicsUpdate() && NextTimeStep > AnimPhysicsMinDeltaTime)
 		{
 			// Calculate gravity direction
@@ -668,7 +700,7 @@ void FAnimNode_AnimDynamics::EvaluateSkeletalControl_AnyThread(FComponentSpacePo
 
 				FTransform NewBoneTransform;
 				
-				if (bScalingEnabled)
+				if (bScalingEnabled && SimulationSpace == AnimPhysSimSpaceType::World && !CVarDisableWorldSpace.GetValueOnAnyThread())
 				{
 					NewBoneTransform = FTransform(CurrentBody.Pose.Orientation, CurrentBody.Pose.Position - CurrentBody.Pose.Orientation.RotateVector(PhysicsBodyJointOffsets[Idx]), FVector(CurrentBody.Pose.Scale));
 				}
@@ -929,7 +961,7 @@ void FAnimNode_AnimDynamics::RequestInitialise(ETeleportType InTeleportType)
 
 void FAnimNode_AnimDynamics::InitPhysics(FComponentSpacePoseContext& Output)
 {
-	const int bScalingSupportEnabled = CVarEnableScaling.GetValueOnAnyThread();
+	const bool bScalingSupportEnabled = CVarEnableScaling.GetValueOnAnyThread() && (SimulationSpace == AnimPhysSimSpaceType::World) && !CVarDisableWorldSpace.GetValueOnAnyThread();
 
 	switch (InitTeleportType)
 	{
@@ -970,9 +1002,12 @@ void FAnimNode_AnimDynamics::InitPhysics(FComponentSpacePoseContext& Output)
 
 			FTransform BodyTransform = GetBoneTransformInSimSpace(Output, PhysicsBodyDef.BoundBone.GetCompactPoseIndex(BoneContainer));
 
+			float BodyScale = BodyTransform.GetScale3D().GetAbsMin();
+			BodyScale = BodyScale < UE_KINDA_SMALL_NUMBER ? UE_KINDA_SMALL_NUMBER : BodyScale;
+
 			if (bScalingSupportEnabled)
 			{
-				BodyShapes.Add(FAnimPhysShape::MakeBox(PhysicsBodyDef.BoxExtents * BodyTransform.GetScale3D()));
+				BodyShapes.Add(FAnimPhysShape::MakeBox(PhysicsBodyDef.BoxExtents * BodyScale));
 				BodyTransform.SetTranslation(BodyTransform.TransformPosition(PhysicsBodyDef.LocalJointOffset));
 			}
 			else
@@ -984,7 +1019,7 @@ void FAnimNode_AnimDynamics::InitPhysics(FComponentSpacePoseContext& Output)
 			FAnimPhysLinkedBody NewChainBody(BodyShapes, BodyTransform.GetTranslation(), PhysicsBodyDef.BoundBone);
 			FAnimPhysRigidBody& PhysicsBody = NewChainBody.RigidBody.PhysBody;
 			PhysicsBody.Pose.Orientation = BodyTransform.GetRotation();
-			PhysicsBody.Pose.Scale = BodyTransform.GetScale3D().X;
+			PhysicsBody.Pose.Scale = BodyScale;
 			PhysicsBody.PreviousOrientation = PhysicsBody.Pose.Orientation;
 			PhysicsBody.NextOrientation = PhysicsBody.Pose.Orientation;
 			PhysicsBody.CollisionType = PhysicsBodyDef.CollisionType;
@@ -992,13 +1027,13 @@ void FAnimNode_AnimDynamics::InitPhysics(FComponentSpacePoseContext& Output)
 			switch (PhysicsBody.CollisionType)
 			{
 			case AnimPhysCollisionType::CustomSphere:
-				PhysicsBody.SphereCollisionRadius = static_cast<float>(bScalingSupportEnabled ? PhysicsBodyDef.SphereCollisionRadius * BodyTransform.GetScale3D().X : PhysicsBodyDef.SphereCollisionRadius);
+				PhysicsBody.SphereCollisionRadius = static_cast<float>(bScalingSupportEnabled ? PhysicsBodyDef.SphereCollisionRadius * BodyScale : PhysicsBodyDef.SphereCollisionRadius);
 				break;
 			case AnimPhysCollisionType::InnerSphere:
-				PhysicsBody.SphereCollisionRadius = static_cast<float>(bScalingSupportEnabled ? (PhysicsBodyDef.BoxExtents.GetAbsMin() * BodyTransform.GetScale3D().X / 2.0) : PhysicsBodyDef.BoxExtents.GetAbsMin() / 2.0);
+				PhysicsBody.SphereCollisionRadius = static_cast<float>(bScalingSupportEnabled ? (PhysicsBodyDef.BoxExtents.GetAbsMin() * BodyScale / 2.0) : PhysicsBodyDef.BoxExtents.GetAbsMin() / 2.0);
 				break;
 			case AnimPhysCollisionType::OuterSphere:
-				PhysicsBody.SphereCollisionRadius = static_cast<float>(bScalingSupportEnabled ? (PhysicsBodyDef.BoxExtents.GetAbsMax() * BodyTransform.GetScale3D().X / 2.0) : PhysicsBodyDef.BoxExtents.GetAbsMax() / 2.0);
+				PhysicsBody.SphereCollisionRadius = static_cast<float>(bScalingSupportEnabled ? (PhysicsBodyDef.BoxExtents.GetAbsMax() * BodyScale / 2.0) : PhysicsBodyDef.BoxExtents.GetAbsMax() / 2.0);
 				break;
 			default:
 				break;
@@ -1671,15 +1706,19 @@ FVector FAnimNode_AnimDynamics::TransformWorldVectorToSimSpaceScaled(FComponentS
 
 void FAnimNode_AnimDynamics::ConvertSimulationSpace(FComponentSpacePoseContext& Output, AnimPhysSimSpaceType From, AnimPhysSimSpaceType To) const
 {
-	for(FAnimPhysRigidBody* Body : BaseBodyPtrs)
+	const int bScalingSupportEnabled = CVarEnableScaling.GetValueOnAnyThread();
+	const bool bWorldDisabled = CVarDisableWorldSpace.GetValueOnAnyThread();
+
+	for(int Idx = 0; Idx < BaseBodyPtrs.Num(); Idx++)
 	{
+		FAnimPhysRigidBody* Body = BaseBodyPtrs[Idx];
 		if(!Body)
 		{
-			return;
+			continue;
 		}
 
 		// Get transform
-		FTransform BodyTransform(Body->Pose.Orientation, Body->Pose.Position);
+		FTransform BodyTransform(Body->Pose.Orientation, Body->Pose.Position, FVector(Body->Pose.Scale));
 		// Out to component space
 		BodyTransform = GetComponentSpaceTransformFromSimSpace(LastSimSpace, Output, BodyTransform);
 		// In to new space
@@ -1688,6 +1727,23 @@ void FAnimNode_AnimDynamics::ConvertSimulationSpace(FComponentSpacePoseContext&
 		// Push back to body
 		Body->Pose.Orientation = BodyTransform.GetRotation();
 		Body->Pose.Position = BodyTransform.GetTranslation();
+
+		float NewScale = BodyTransform.GetScale3D().GetAbsMin();
+		NewScale = NewScale < UE_KINDA_SMALL_NUMBER ? UE_KINDA_SMALL_NUMBER : NewScale;
+
+		if (bScalingSupportEnabled)
+		{
+			if (From == AnimPhysSimSpaceType::World && !bWorldDisabled && To != AnimPhysSimSpaceType::World)
+			{
+				UpdateBodyScaling(Idx, 1);
+			}
+			else if (From != AnimPhysSimSpaceType::World && To == AnimPhysSimSpaceType::World && !bWorldDisabled)
+			{
+				UpdateBodyScaling(Idx, NewScale * Body->Pose.Scale);
+			}
+		}
+
+		Body->Pose.Scale = NewScale;
 	}
 }
 
diff --git a/Engine/Source/Runtime/AnimGraphRuntime/Public/BoneControllers/AnimNode_AnimDynamics.h b/Engine/Source/Runtime/AnimGraphRuntime/Public/BoneControllers/AnimNode_AnimDynamics.h
--- a/Engine/Source/Runtime/AnimGraphRuntime/Public/BoneControllers/AnimNode_AnimDynamics.h
+++ b/Engine/Source/Runtime/AnimGraphRuntime/Public/BoneControllers/AnimNode_AnimDynamics.h
@@ -408,6 +408,12 @@ struct FAnimNode_AnimDynamics : public FAnimNode_SkeletalControlBase
 	UPROPERTY(EditAnywhere, BlueprintReadWrite, Category = PhysicsParameters, meta = (DisplayAfter = "PhysicsBodyDefinitions", EditCondition = bAngularSpring, PinHiddenByDefault))
 	float AngularSpringConstant;
 
+	/**
+		For world-space simulations, if the magnitude of the component's 3D scale is less than WorldSpaceMinimumScale, do not update the node.
+	*/
+	UPROPERTY(EditAnywhere, Category = Settings)
+	float WorldSpaceMinimumScale;
+
 	/** Scale to apply to calculated wind velocities in the solver */
 	UPROPERTY(EditAnywhere, Category = Wind, meta=(DisplayAfter="bEnableWind"))
 	float WindScale;
@@ -620,6 +626,8 @@ private:
 	void CalculateSimulationSpaceMotion(AnimPhysSimSpaceType Space, FAnimPhysSimSpaceSettings Settings, FComponentSpacePoseContext& Output, const float Dt,
 		FVector& SpaceAngularVel, FVector& SpaceAngularAcc);
 
+	void UpdateBodyScaling(const int32 BodyIndex, float Scale) const;
+
 	// Maximum time to consider when accumulating time debt to avoid spiraling
 	static ANIMGRAPHRUNTIME_API const float MaxTimeDebt;
 
