Fix scaling body properties on world space for AnimDynamics. Add teleport support on scale changes. Fix debug draw for scaled bodies.

diff --git a/Engine/Source/Runtime/AnimGraphRuntime/Private/BoneControllers/AnimNode_AnimDynamics.cpp b/Engine/Source/Runtime/AnimGraphRuntime/Private/BoneControllers/AnimNode_AnimDynamics.cpp
--- a/Engine/Source/Runtime/AnimGraphRuntime/Private/BoneControllers/AnimNode_AnimDynamics.cpp
+++ b/Engine/Source/Runtime/AnimGraphRuntime/Private/BoneControllers/AnimNode_AnimDynamics.cpp
@@ -29,6 +29,8 @@ TAutoConsoleVariable<int32> CVarEnableWind(TEXT("p.AnimDynamicsWind"), 1, TEXT("
 TAutoConsoleVariable<float> CVarComponentAppliedLinearAccClampOverride(TEXT("p.AnimDynamics.ComponentAppliedLinearAccClampOverride"), -1.0f, TEXT("Override the per asset setting for all axis (X,Y & Z) of ComponentAppliedLinearAccClamp for all Anim Dynamics Nodes. Negative values are ignored."));
 TAutoConsoleVariable<float> CVarGravityScale(TEXT("p.AnimDynamics.GravityScale"), 1.0f, TEXT("Multiplies the default gravity and the gravity override on all Anim Dynamics Nodes."));
 TAutoConsoleVariable<bool>  CVarDisableWorldSpace(TEXT("p.AnimDynamics.DisableWorldSpace"), false, TEXT("Disable World Sapce transform and override it to component space on all Anim Dynamics Nodes."));
+TAutoConsoleVariable<bool>  CVarScaleInitialBodyMassByComponentSpace(TEXT("p.AnimDynamics.ScaleInitialBodyMassByComponentSpace"), true, TEXT("Scale Initial Body Mass By Component Space Scale when using World Space."));
+TAutoConsoleVariable<bool>  CVarTeleportBodiesOnScaleChange(TEXT("p.AnimDynamics.TeleportBodiesOnScaleChange"), true, TEXT("Teleport Bodies On Scale Change when using World Space."));
 
 // FSimSpaceSettingsWithControl forced overrides for testing
 bool bAnimDynamics_SimSpace_EnableOverride = false;
@@ -115,7 +117,7 @@ void FAnimNode_AnimDynamics::DrawBodies(FComponentSpacePoseContext& InContext, c
 		FTransform Transform(Body.Pose.Orientation, Body.Pose.Position, SimulationSpace != AnimPhysSimSpaceType::World || CVarDisableWorldSpace.GetValueOnAnyThread() ? FVector(Body.Pose.Scale) : FVector(1));
 		Transform = GetComponentSpaceTransformFromSimSpace(SimulationSpace, InContext, Transform);
 
-		Proxy->AnimDrawDebugCoordinateSystem(Transform.GetTranslation(), Transform.Rotator(), 2.0f, false, -1.0f, 0.15f);
+		Proxy->AnimDrawDebugCoordinateSystem(ToWorldV(InContext, Transform.GetTranslation()), Transform.Rotator(), 2.0f * Body.Pose.Scale, false, -1.0f, 0.15f);
 
 		for(const FAnimPhysShape& Shape : Body.Shapes)
 		{
@@ -453,12 +455,15 @@ void FAnimNode_AnimDynamics::EvaluateSkeletalControl_AnyThread(FComponentSpacePo
 			InitTeleportType = ETeleportType::None;
 		}
 
-		const FTransform SimWorldTransform = GetSimToWorldSpaceTransform(SimulationSpace, Output);
+		const FTransform SimWorldTransform = GetSimToWorldSpaceTransform(SimulationSpace != AnimPhysSimSpaceType::World ? SimulationSpace : AnimPhysSimSpaceType::Component, Output);
 
 		if (bScalingEnabled)
 		{
-			if (!FMath::IsNearlyEqual(SimWorldTransform.GetScale3D().X, PreviousSimSpaceTM.GetScale3D().X))
+			if (!SimWorldTransform.GetScale3D().Equals(PreviousSimSpaceTM.GetScale3D(), UE_KINDA_SMALL_NUMBER))
 			{
+				const bool bWorldSpace = SimulationSpace == AnimPhysSimSpaceType::World && !CVarDisableWorldSpace.GetValueOnAnyThread();
+				const bool bTeleportEnabled = bWorldSpace && CVarTeleportBodiesOnScaleChange.GetValueOnAnyThread();
+
 				// If the scaling of the sim space has changed, update the bodies to the proper scale
 				const FBoneContainer& BoneContainer = Output.Pose.GetPose().GetBoneContainer();
 				for (int32 ActiveIndex : ActiveBoneIndices)
@@ -480,14 +485,29 @@ void FAnimNode_AnimDynamics::EvaluateSkeletalControl_AnyThread(FComponentSpacePo
 
 					float Scale = BoundBoneTransform.GetScale3D().GetAbsMin();
 					Scale = Scale < UE_KINDA_SMALL_NUMBER ? UE_KINDA_SMALL_NUMBER : Scale;
-					
+
 					if (RigidBody.Pose.Scale != Scale)
 					{
-						if (SimulationSpace == AnimPhysSimSpaceType::World && !CVarDisableWorldSpace.GetValueOnAnyThread())
+						if (bWorldSpace)
 						{
 							UpdateBodyScaling(ActiveIndex, Scale);
 						}
 						RigidBody.Pose.Scale = Scale;
+
+						if (bTeleportEnabled)
+						{
+							FTransform BodyTransform(RigidBody.Pose.Orientation, RigidBody.Pose.Position - RigidBody.Pose.Orientation.RotateVector(PhysicsBodyDefinitions[ActiveIndex].LocalJointOffset));
+							BodyTransform = GetComponentSpaceTransformFromSimSpace(SimulationSpace, Output, BodyTransform, PreviousCompWorldSpaceTM, PreviousActorWorldSpaceTM);
+
+							// move to new space
+							BodyTransform = GetSimSpaceTransformFromComponentSpace(SimulationSpace, Output, BodyTransform);
+
+							RigidBody.Pose.Orientation = BodyTransform.GetRotation();
+							RigidBody.PreviousOrientation = RigidBody.Pose.Orientation;
+							RigidBody.NextOrientation = RigidBody.Pose.Orientation;
+
+							RigidBody.Pose.Position = BodyTransform.GetTranslation() + RigidBody.Pose.Orientation.RotateVector(PhysicsBodyDefinitions[ActiveIndex].LocalJointOffset);
+						}
 					}
 				}
 			}
@@ -497,6 +517,7 @@ void FAnimNode_AnimDynamics::EvaluateSkeletalControl_AnyThread(FComponentSpacePo
 		while(BodiesToReset.Num() > 0)
 		{
 			FAnimPhysLinkedBody* BodyToReset = BodiesToReset.Pop(EAllowShrinking::No);
+
 			if(BodyToReset && BodyToReset->RigidBody.BoundBone.IsValidToEvaluate(RequiredBones))
 			{
 				FTransform BoneTransform = GetBoneTransformInSimSpace(Output, BodyToReset->RigidBody.BoundBone.GetCompactPoseIndex(RequiredBones));
@@ -743,6 +764,8 @@ void FAnimNode_AnimDynamics::EvaluateSkeletalControl_AnyThread(FComponentSpacePo
 			}
 		}
 
+		PreviousSimSpaceTM = SimWorldTransform;
+
 		// Store our sim space in case it changes
 		LastSimSpace = SimulationSpace;
 
@@ -994,6 +1017,10 @@ void FAnimNode_AnimDynamics::InitPhysics(FComponentSpacePoseContext& Output)
 		FTransform PreviousBodyTransformSimSpace;
 		PreviousBodyTransformSimSpace.SetIdentity();
 
+		const FTransform InitialComponentTransform = GetSimToWorldSpaceTransform(AnimPhysSimSpaceType::Component, Output);
+		float InitialComponentScale = InitialComponentTransform.GetScale3D().GetAbsMin();
+		InitialComponentScale = InitialComponentScale < UE_KINDA_SMALL_NUMBER ? UE_KINDA_SMALL_NUMBER : InitialComponentScale;
+
 		for (FAnimPhysBodyDefinition& PhysicsBodyDef : PhysicsBodyDefinitions)
 		{
 			TArray<FAnimPhysShape> BodyShapes;
@@ -1004,7 +1031,7 @@ void FAnimNode_AnimDynamics::InitPhysics(FComponentSpacePoseContext& Output)
 
 			float BodyScale = BodyTransform.GetScale3D().GetAbsMin();
 			BodyScale = BodyScale < UE_KINDA_SMALL_NUMBER ? UE_KINDA_SMALL_NUMBER : BodyScale;
-
+			
 			if (bScalingSupportEnabled)
 			{
 				BodyShapes.Add(FAnimPhysShape::MakeBox(PhysicsBodyDef.BoxExtents * BodyScale));
@@ -1024,6 +1051,12 @@ void FAnimNode_AnimDynamics::InitPhysics(FComponentSpacePoseContext& Output)
 			PhysicsBody.NextOrientation = PhysicsBody.Pose.Orientation;
 			PhysicsBody.CollisionType = PhysicsBodyDef.CollisionType;
 
+			//Masses originally were always 1 regardless of Scale or BodyShape Bounds. Allow the user to enable mass scaling so masses are adjusted based on initial scale.
+			if (bScalingSupportEnabled && !CVarDisableWorldSpace.GetValueOnAnyThread()  && CVarScaleInitialBodyMassByComponentSpace.GetValueOnAnyThread())
+			{
+				FAnimPhys::ScaleRigidBodyMass(&PhysicsBody, InitialComponentScale);
+			}
+
 			switch (PhysicsBody.CollisionType)
 			{
 			case AnimPhysCollisionType::CustomSphere:
@@ -1141,7 +1174,8 @@ void FAnimNode_AnimDynamics::InitPhysics(FComponentSpacePoseContext& Output)
 	InitTeleportType = ETeleportType::None;
 	PreviousCompWorldSpaceTM = Output.AnimInstanceProxy->GetComponentTransform();
 	PreviousActorWorldSpaceTM = Output.AnimInstanceProxy->GetActorTransform();
-	PreviousSimSpaceTM = GetSimToWorldSpaceTransform(SimulationSpace, Output);
+
+	PreviousSimSpaceTM = GetSimToWorldSpaceTransform(SimulationSpace != AnimPhysSimSpaceType::World ? SimulationSpace : AnimPhysSimSpaceType::Component, Output);
 }
 
 void FAnimNode_AnimDynamics::TermPhysics()
@@ -1172,6 +1206,8 @@ void FAnimNode_AnimDynamics::UpdateLimits(FComponentSpacePoseContext& Output)
 
 	const FBoneContainer& BoneContainer = Output.Pose.GetPose().GetBoneContainer();
 
+	const bool bShouldScale = SimulationSpace == AnimPhysSimSpaceType::World && !CVarDisableWorldSpace.GetValueOnAnyThread() && CVarEnableScaling.GetValueOnAnyThread();
+
 	for(int32 ActiveIndex : ActiveBoneIndices)
 	{
 		const FAnimPhysBodyDefinition& PhysicsBodyDef = PhysicsBodyDefinitions[ActiveIndex];
@@ -1200,11 +1236,11 @@ void FAnimNode_AnimDynamics::UpdateLimits(FComponentSpacePoseContext& Output)
 		
 		if (PrevBody)
 		{
-			ShapeTransform = FTransform(FQuat::Identity, ConstraintOffsets[ActiveIndex].Body0Offset);
+			ShapeTransform = FTransform(FQuat::Identity, ConstraintOffsets[ActiveIndex].Body0Offset * (bShouldScale ? PrevBody->Pose.Scale : 1));
 		}
 
 		// Local offset to joint for Body1
-		const FVector Body1JointOffset = ConstraintOffsets[ActiveIndex].Body1Offset;
+		const FVector Body1JointOffset = ConstraintOffsets[ActiveIndex].Body1Offset * (bShouldScale ? RigidBody.Pose.Scale : 1);
 
 		if (PhysicsBodyDef.ConstraintSetup.bLinearFullyLocked)
 		{
@@ -1286,10 +1322,10 @@ void FAnimNode_AnimDynamics::UpdateLimits(FComponentSpacePoseContext& Output)
 				switch(SphericalLimit.LimitType)
 				{
 				case ESphericalLimitType::Inner:
-					FAnimPhys::ConstrainSphericalInner(NextTimeStep, LinearLimits, &RigidBody, SphereTransform, SphericalLimit.LimitRadius);
+					FAnimPhys::ConstrainSphericalInner(NextTimeStep, LinearLimits, &RigidBody, SphereTransform, SphericalLimit.LimitRadius * (bShouldScale ? RigidBody.Pose.Scale : 1));
 					break;
 				case ESphericalLimitType::Outer:
-					FAnimPhys::ConstrainSphericalOuter(NextTimeStep, LinearLimits, &RigidBody, SphereTransform, SphericalLimit.LimitRadius);
+					FAnimPhys::ConstrainSphericalOuter(NextTimeStep, LinearLimits, &RigidBody, SphereTransform, SphericalLimit.LimitRadius * (bShouldScale ? RigidBody.Pose.Scale : 1));
 					break;
 				default:
 					break;
