Remove Curve Node

Anim Graph Node that removes animation curves.
It's very useful in certain layering scenarios and when the same animation content is used in separate parts of the graph or to eliminate user created curves that should never be added in a certain place of the graph.

- can remove by list of curve names provided (exposable);
- can remove all curves;
- in KeepOnlySpecifiedCurves  mode - exception curve list can be given;
- option to affect\not affect Regular \ Material type and Morph Target type curves seprately;

#changelist validated

diff --git a/Engine/Source/Editor/AnimGraph/Private/AnimGraphNode_RemoveCurve.cpp b/Engine/Source/Editor/AnimGraph/Private/AnimGraphNode_RemoveCurve.cpp
new file mode 100644
--- /dev/null
+++ b/Engine/Source/Editor/AnimGraph/Private/AnimGraphNode_RemoveCurve.cpp
@@ -0,0 +1,26 @@
+#include "AnimGraphNode_RemoveCurve.h"
+
+#include UE_INLINE_GENERATED_CPP_BY_NAME(AnimGraphNode_RemoveCurve)
+
+#define LOCTEXT_NAMESPACE "RemoveCurve"
+
+UAnimGraphNode_RemoveCurve::UAnimGraphNode_RemoveCurve()
+{
+}
+
+FText UAnimGraphNode_RemoveCurve::GetMenuCategory() const
+{
+	return LOCTEXT("AnimGraphNode_RemoveCurve_Category", "Animation|Curves");
+}
+
+FText UAnimGraphNode_RemoveCurve::GetTooltipText() const
+{	
+	return LOCTEXT("AnimGraphNode_RemoveCurve_Category", "Removes animation curves from the pose input");
+}
+
+FText UAnimGraphNode_RemoveCurve::GetNodeTitle(ENodeTitleType::Type TitleType) const
+{
+	return LOCTEXT("AnimGraphNode_RemoveCurve_Title", "Remove Curve");
+}
+
+#undef LOCTEXT_NAMESPACE
\ No newline at end of file
diff --git a/Engine/Source/Editor/AnimGraph/Public/AnimGraphNode_RemoveCurve.h b/Engine/Source/Editor/AnimGraph/Public/AnimGraphNode_RemoveCurve.h
new file mode 100644
--- /dev/null
+++ b/Engine/Source/Editor/AnimGraph/Public/AnimGraphNode_RemoveCurve.h
@@ -0,0 +1,28 @@
+#pragma once
+
+#include "CoreMinimal.h"
+#include "AnimGraphNode_Base.h"
+#include "AnimNodes/AnimNode_RemoveCurve.h"
+#include "AnimGraphNode_RemoveCurve.generated.h"
+
+///** Removes animation curves on a pose */
+UCLASS(MinimalAPI, Experimental)
+class UAnimGraphNode_RemoveCurve : public UAnimGraphNode_Base
+{
+	GENERATED_BODY()
+
+	UPROPERTY(EditAnywhere, Category = Settings)
+	FAnimNode_RemoveCurve Node;
+
+public:	
+
+	UAnimGraphNode_RemoveCurve();
+
+	// UEdGraphNode interface
+	virtual FText GetTooltipText() const override;
+	virtual FText GetNodeTitle(ENodeTitleType::Type TitleType) const override;
+	virtual FText GetMenuCategory() const override;
+	// End of UEdGraphNode interface
+
+	
+};
\ No newline at end of file
diff --git a/Engine/Source/Runtime/AnimGraphRuntime/Private/AnimNodes/AnimNode_RemoveCurve.cpp b/Engine/Source/Runtime/AnimGraphRuntime/Private/AnimNodes/AnimNode_RemoveCurve.cpp
new file mode 100644
--- /dev/null
+++ b/Engine/Source/Runtime/AnimGraphRuntime/Private/AnimNodes/AnimNode_RemoveCurve.cpp
@@ -0,0 +1,207 @@
+#include "AnimNodes/AnimNode_RemoveCurve.h"
+#include "Animation/AnimInstanceProxy.h"
+
+
+
+#include UE_INLINE_GENERATED_CPP_BY_NAME(AnimNode_RemoveCurve)
+
+FAnimNode_RemoveCurve::FAnimNode_RemoveCurve()
+{	
+	bEnabled = true;
+	Curves = {};
+	RemoveMode = ERemoveCurveMode::RemoveSpecifiedCurves;
+	bAffectMorphTargetCurves = true;
+	bAffectMaterialCurves = true;
+	bAffectRegularCurves = true;
+	NumCurvesOnSkeleton = 0;
+}
+
+void FAnimNode_RemoveCurve::Initialize_AnyThread(const FAnimationInitializeContext& Context)
+{
+	DECLARE_SCOPE_HIERARCHICAL_COUNTER_ANIMNODE(Initialize_AnyThread)
+	Super::Initialize_AnyThread(Context);
+	SourcePose.Initialize(Context);
+
+	if (!bEnabled)
+	{
+		return;
+	}
+
+	// reserve expected amount
+	int ExcludedCurvesAmount = 0;
+	const bool bRemoveAllCurves = RemoveMode == ERemoveCurveMode::KeepOnlySpecifiedCurves ? true : false;
+
+	if (bRemoveAllCurves)
+	{
+		// technically, it's possible that pose has more curves than those that are specified on skeleton, but afaik unreal will allocate more automaticaally if that happens.
+		ExcludedCurvesAmount = NumCurvesOnSkeleton;
+	}
+	else
+	{
+		ExcludedCurvesAmount = Curves.Num();
+	}
+	
+	CurvesForRemoval.Reserve(ExcludedCurvesAmount);
+}
+
+void FAnimNode_RemoveCurve::CacheBones_AnyThread(const FAnimationCacheBonesContext& Context)
+{
+	DECLARE_SCOPE_HIERARCHICAL_COUNTER_ANIMNODE(CacheBones_AnyThread)
+	Super::CacheBones_AnyThread(Context);
+	SourcePose.CacheBones(Context);
+
+	if (!bEnabled)
+	{
+		return;
+	}
+	
+	// Store curve types on the skel to look up in Eval state
+	MorphTargetCurveNames.Reset();
+	MaterialCurveNames.Reset();
+
+	// get skeleton's curves and store curves of different types to different lists
+	const USkeleton* Skel = Context.AnimInstanceProxy ? Context.AnimInstanceProxy->GetSkeleton() : nullptr;
+	if (!Skel)
+	{
+		return;
+	}
+	
+	TArray<FName> MetaNames;
+	Skel->GetCurveMetaDataNames(MetaNames);
+	for (const FName& N : MetaNames)
+	{
+		if (const FCurveMetaData* Meta = Skel->GetCurveMetaData(N))
+		{
+			if (Meta->Type.bMorphtarget)
+			{
+				MorphTargetCurveNames.Add(N);
+			}
+
+			if (Meta->Type.bMaterial)
+			{
+				MaterialCurveNames.Add(N);
+			}			
+		}
+	}
+
+	NumCurvesOnSkeleton = MetaNames.Num();
+}
+
+void FAnimNode_RemoveCurve::Update_AnyThread(const FAnimationUpdateContext& Context)
+{
+	DECLARE_SCOPE_HIERARCHICAL_COUNTER_ANIMNODE(Update_AnyThread)
+	SourcePose.Update(Context);
+	GetEvaluateGraphExposedInputs().Execute(Context);
+
+}
+
+void FAnimNode_RemoveCurve::Evaluate_AnyThread(FPoseContext& Output)
+{
+	DECLARE_SCOPE_HIERARCHICAL_COUNTER_ANIMNODE(Evaluate_AnyThread)
+	ANIM_MT_SCOPE_CYCLE_COUNTER_VERBOSE(RemoveCurve, !IsInGameThread());
+
+	// Evaluate upstream pose first
+	SourcePose.Evaluate(Output);
+
+	if (!bEnabled)
+	{
+		return;
+	}
+
+	const bool bRemoveAllCurves = RemoveMode == ERemoveCurveMode::KeepOnlySpecifiedCurves ? true : false;
+
+	if (!bRemoveAllCurves && Curves.Num() == 0)
+	{
+		return;
+	}
+	
+	CurvesForRemoval = {};
+
+	// helper func to check if curve is not allowed to be removed	
+	auto CurveIsOfAllowedToRemoveType = [this](const FName& CurveName) -> bool
+	{
+		// not allowed to be removed cause the type is checked off on the node
+		bool bAllowedIfMorphTargetCurve = bAffectMorphTargetCurves ? true : !MorphTargetCurveNames.Contains(CurveName);
+		bool bAllowedIfMaterialCurve = bAffectMaterialCurves ? true : !MaterialCurveNames.Contains(CurveName);
+		bool bAllowedIfRegularCurve = bAffectRegularCurves ? true : (MorphTargetCurveNames.Contains(CurveName) || MaterialCurveNames.Contains(CurveName));
+
+		return bAllowedIfMorphTargetCurve && bAllowedIfMaterialCurve && bAllowedIfRegularCurve;
+	};
+
+	// If Removal of all curves is requested
+	if (bRemoveAllCurves)
+	{
+		CurvesForRemoval.Reserve(Output.Curve.Num());
+		
+		Output.Curve.ForEachElement([this, CurveIsOfAllowedToRemoveType](const UE::Anim::FCurveElement& Elem)
+		{	
+			
+			// is in exceptions list
+			bool bIsInExceptions = Curves.Contains(Elem.Name);
+
+			if (CurveIsOfAllowedToRemoveType(Elem.Name) && !bIsInExceptions)
+			{
+				CurvesForRemoval.Add(Elem.Name);
+			}
+		});
+	}
+	else
+	{
+		CurvesForRemoval.Reserve(Curves.Num());
+
+		// iterate through user provided list of curve names and remove
+		for (const FName& CurveName : Curves)
+		{	
+			
+			if (CurveIsOfAllowedToRemoveType(CurveName))
+			{
+				CurvesForRemoval.Add(CurveName);
+			}
+		}			
+	}
+
+	// Remove each collected and aproved curve from this pose's curve container	
+	UE::Anim::FNamedValueArrayUtils::RemoveByPredicate(
+		Output.Curve,
+		CurvesForRemoval,
+		[](const UE::Anim::FCurveElement& OutputCurve, const UE::Anim::FCurveElement& CurveToRemove)
+		{			
+			return true;
+		}
+	);
+
+}
+
+void FAnimNode_RemoveCurve::GatherDebugData(FNodeDebugData& DebugData)
+{
+	DECLARE_SCOPE_HIERARCHICAL_COUNTER_ANIMNODE(GatherDebugData)
+	FString DebugLine = DebugData.GetNodeName(this);
+
+	DebugLine += bEnabled ? TEXT("True") : TEXT("False");
+	DebugData.AddDebugItem(DebugLine);
+
+	DebugLine = TEXT("Affects MorphTarget Type curves: ");
+	DebugLine += bAffectMorphTargetCurves ? TEXT("True") : TEXT("False");
+	DebugData.AddDebugItem(DebugLine);
+
+	DebugLine = TEXT("Affects Material Type curves: ");
+	DebugLine += bAffectMaterialCurves ? TEXT("True") : TEXT("False");
+	DebugData.AddDebugItem(DebugLine);
+
+	DebugLine = TEXT("Affects Regular Type curves: ");
+	DebugLine += bAffectMaterialCurves ? TEXT("True") : TEXT("False");
+	DebugData.AddDebugItem(DebugLine);
+
+	DebugLine = TEXT(" Removed=[");
+	CurvesForRemoval.ForEachElement([&DebugLine](const UE::Anim::FCurveElement& Elem)
+		{
+			DebugLine += Elem.Name.ToString();
+			DebugLine += TEXT(", ");
+		}
+	);
+
+	DebugLine += TEXT("]");
+	DebugData.AddDebugItem(DebugLine);
+
+	SourcePose.GatherDebugData(DebugData);
+}
\ No newline at end of file
diff --git a/Engine/Source/Runtime/AnimGraphRuntime/Public/AnimNodes/AnimNode_RemoveCurve.h b/Engine/Source/Runtime/AnimGraphRuntime/Public/AnimNodes/AnimNode_RemoveCurve.h
new file mode 100644
--- /dev/null
+++ b/Engine/Source/Runtime/AnimGraphRuntime/Public/AnimNodes/AnimNode_RemoveCurve.h
@@ -0,0 +1,84 @@
+#pragma once
+
+#include "Animation/AnimNodeBase.h"
+#include "Animation/AnimCurveTypes.h"
+#include "Animation/NamedValueArray.h"
+#include "CoreMinimal.h"
+#include "AnimNode_RemoveCurve.generated.h"
+
+UENUM()
+enum class ERemoveCurveMode : uint8
+{
+	/** Removes animation curves which names are specified in the Curves variable from the Pose Input. */
+	RemoveSpecifiedCurves,
+
+	/** Removes ALL animation curves in the Pose Input, EXCEPT those specified in the Curves variable */
+	KeepOnlySpecifiedCurves
+};
+
+
+//** Removes animation curves on a pose */
+USTRUCT(BlueprintInternalUseOnly, Experimental)
+struct FAnimNode_RemoveCurve : public FAnimNode_Base
+{
+	GENERATED_BODY()
+
+public:
+	// Input link (the upstream pose)
+	UPROPERTY(EditAnywhere, BlueprintReadWrite, Category = "Links")
+	FPoseLink SourcePose;
+
+	// Is Enabled
+	UPROPERTY(EditAnywhere, BlueprintReadWrite, Category = "Alpha", meta = (PinShownByDefault))
+	bool bEnabled;
+
+	/** Node mode selection.
+	RemoveSpecifiedCurves -  will remove animation curves which names are specified in the Curves variable from the Pose Input
+	KeepOnlySpecifiedCurves - will remove ALL animation curves in the Pose Input, EXCEPT those specified in the Curves variable	*/
+	UPROPERTY(EditAnywhere, BlueprintReadWrite, Category = "Remove Curves", meta = (BlueprintCompilerGeneratedDefaults, PinHiddenByDefault))
+	ERemoveCurveMode RemoveMode;
+
+	/** List of animation curves.*/	
+	UPROPERTY(EditAnywhere, BlueprintReadWrite, Category = "Remove Curves", meta = (BlueprintCompilerGeneratedDefaults, PinHiddenByDefault))
+	TSet<FName> Curves;
+
+	/** Affect curves of a Morph Target type when True.
+	Ignores them if False. I.e. curves of Morph Target type will be left completely unaffected by this node*/
+	UPROPERTY(EditAnywhere, BlueprintReadWrite, Category = "Remove Curves Filter by Type", meta = (NeverAsPin))
+	bool bAffectMorphTargetCurves;
+
+	/** Affect curves of a Material type. 
+	Ignores them if False. I.e. curves of Material type will be left completely unaffected by this node */
+	UPROPERTY(EditAnywhere, BlueprintReadWrite, Category = "Remove Curves Filter by Type", meta = (NeverAsPin))
+	bool bAffectMaterialCurves;
+
+	/** Affect normal animation curves (not Material or Morph Target type). 
+	Ignores them if False. I.e. regualar animation curves will be left completely unaffected by this node */
+	UPROPERTY(EditAnywhere, BlueprintReadWrite, Category = "Remove Curves Filter by Type", meta = (NeverAsPin))
+	bool bAffectRegularCurves;
+
+	// Internal vars.
+	UPROPERTY(transient)
+	TArray<FName> MorphTargetCurveNames;
+
+	UPROPERTY(transient)
+	TArray<FName> MaterialCurveNames;
+
+	UPROPERTY(transient)
+	int32 NumCurvesOnSkeleton;
+
+	ANIMGRAPHRUNTIME_API FAnimNode_RemoveCurve();
+
+	// FAnimNode_Base interface
+	virtual void Initialize_AnyThread(const FAnimationInitializeContext& Context) override;
+	ANIMGRAPHRUNTIME_API virtual void CacheBones_AnyThread(const FAnimationCacheBonesContext& Context) override;
+	ANIMGRAPHRUNTIME_API virtual void Evaluate_AnyThread(FPoseContext& Output) override;
+
+	ANIMGRAPHRUNTIME_API virtual void Update_AnyThread(const FAnimationUpdateContext& Context) override;
+	ANIMGRAPHRUNTIME_API virtual void GatherDebugData(FNodeDebugData& DebugData) override;
+	// End of FAnimNode_Base interface
+
+private:
+	// Final list of curves to remove. Global, so it can be used in debug printouts anywhere.
+	UE::Anim::TNamedValueArray<FDefaultAllocator, UE::Anim::FCurveElement> CurvesForRemoval;
+};
