Remove Curve Node

Fix Horde errors:
- add copyright boilerplate;
- add ANIMGRAPHRUNTIME_API macro to Initialize_AnyThread;

diff --git a/Engine/Source/Editor/AnimGraph/Private/AnimGraphNode_RemoveCurve.cpp b/Engine/Source/Editor/AnimGraph/Private/AnimGraphNode_RemoveCurve.cpp
--- a/Engine/Source/Editor/AnimGraph/Private/AnimGraphNode_RemoveCurve.cpp
+++ b/Engine/Source/Editor/AnimGraph/Private/AnimGraphNode_RemoveCurve.cpp
@@ -1,3 +1,5 @@
+// Copyright Epic Games, Inc. All Rights Reserved.
+
 #include "AnimGraphNode_RemoveCurve.h"
 
 #include UE_INLINE_GENERATED_CPP_BY_NAME(AnimGraphNode_RemoveCurve)
diff --git a/Engine/Source/Editor/AnimGraph/Public/AnimGraphNode_RemoveCurve.h b/Engine/Source/Editor/AnimGraph/Public/AnimGraphNode_RemoveCurve.h
--- a/Engine/Source/Editor/AnimGraph/Public/AnimGraphNode_RemoveCurve.h
+++ b/Engine/Source/Editor/AnimGraph/Public/AnimGraphNode_RemoveCurve.h
@@ -1,3 +1,5 @@
+// Copyright Epic Games, Inc. All Rights Reserved.
+
 #pragma once
 
 #include "CoreMinimal.h"
diff --git a/Engine/Source/Runtime/AnimGraphRuntime/Private/AnimNodes/AnimNode_RemoveCurve.cpp b/Engine/Source/Runtime/AnimGraphRuntime/Private/AnimNodes/AnimNode_RemoveCurve.cpp
--- a/Engine/Source/Runtime/AnimGraphRuntime/Private/AnimNodes/AnimNode_RemoveCurve.cpp
+++ b/Engine/Source/Runtime/AnimGraphRuntime/Private/AnimNodes/AnimNode_RemoveCurve.cpp
@@ -1,8 +1,8 @@
+// Copyright Epic Games, Inc. All Rights Reserved.
+
 #include "AnimNodes/AnimNode_RemoveCurve.h"
 #include "Animation/AnimInstanceProxy.h"
 
-
-
 #include UE_INLINE_GENERATED_CPP_BY_NAME(AnimNode_RemoveCurve)
 
 FAnimNode_RemoveCurve::FAnimNode_RemoveCurve()
diff --git a/Engine/Source/Runtime/AnimGraphRuntime/Public/AnimNodes/AnimNode_RemoveCurve.h b/Engine/Source/Runtime/AnimGraphRuntime/Public/AnimNodes/AnimNode_RemoveCurve.h
--- a/Engine/Source/Runtime/AnimGraphRuntime/Public/AnimNodes/AnimNode_RemoveCurve.h
+++ b/Engine/Source/Runtime/AnimGraphRuntime/Public/AnimNodes/AnimNode_RemoveCurve.h
@@ -1,3 +1,5 @@
+// Copyright Epic Games, Inc. All Rights Reserved.
+
 #pragma once
 
 #include "Animation/AnimNodeBase.h"
@@ -70,7 +72,7 @@ public:
 	ANIMGRAPHRUNTIME_API FAnimNode_RemoveCurve();
 
 	// FAnimNode_Base interface
-	virtual void Initialize_AnyThread(const FAnimationInitializeContext& Context) override;
+	ANIMGRAPHRUNTIME_API virtual void Initialize_AnyThread(const FAnimationInitializeContext& Context) override;
 	ANIMGRAPHRUNTIME_API virtual void CacheBones_AnyThread(const FAnimationCacheBonesContext& Context) override;
 	ANIMGRAPHRUNTIME_API virtual void Evaluate_AnyThread(FPoseContext& Output) override;
 
