AnimDynamics Scaling Instability Fixes

diff --git a/Engine/Source/Runtime/AnimGraphRuntime/Private/BoneControllers/AnimNode_AnimDynamics.cpp b/Engine/Source/Runtime/AnimGraphRuntime/Private/BoneControllers/AnimNode_AnimDynamics.cpp
--- a/Engine/Source/Runtime/AnimGraphRuntime/Private/BoneControllers/AnimNode_AnimDynamics.cpp
+++ b/Engine/Source/Runtime/AnimGraphRuntime/Private/BoneControllers/AnimNode_AnimDynamics.cpp
@@ -31,6 +31,7 @@ TAutoConsoleVariable<float> CVarGravityScale(TEXT("p.AnimDynamics.GravityScale")
 TAutoConsoleVariable<bool>  CVarDisableWorldSpace(TEXT("p.AnimDynamics.DisableWorldSpace"), false, TEXT("Disable World Sapce transform and override it to component space on all Anim Dynamics Nodes."));
 TAutoConsoleVariable<bool>  CVarScaleInitialBodyMassByComponentSpace(TEXT("p.AnimDynamics.ScaleInitialBodyMassByComponentSpace"), true, TEXT("Scale Initial Body Mass By Component Space Scale when using World Space."));
 TAutoConsoleVariable<bool>  CVarTeleportBodiesOnScaleChange(TEXT("p.AnimDynamics.TeleportBodiesOnScaleChange"), true, TEXT("Teleport Bodies On Scale Change when using World Space."));
+TAutoConsoleVariable<bool>  CVarScaleWorldSpaceSprings(TEXT("p.AnimDynamics.ScaleWorldSpaceSprings"), true, TEXT("Scale spring constants to match body scale changes when in world space."));
 
 // FSimSpaceSettingsWithControl forced overrides for testing
 bool bAnimDynamics_SimSpace_EnableOverride = false;
@@ -427,6 +428,8 @@ void FAnimNode_AnimDynamics::UpdateBodyScaling(const int32 BodyIndex, float Scal
 			break;
 		}
 		RigidBody->Shapes = BodyShapes;
+
+		FAnimPhys::RigidBodyShapeScaleUpdate(RigidBody, Scale / RigidBody->Pose.Scale);
 	}
 }
 
@@ -1095,8 +1098,8 @@ void FAnimNode_AnimDynamics::InitPhysics(FComponentSpacePoseContext& Output)
 				// Link to parent
 				NewChainBody.ParentBody = &Bodies.Last().RigidBody;
 
-				// Calculate constraint offset positions in the space of each body.
-				const FVector ConstaintLocationSimSpace = (BodyTransform.GetTranslation() - PreviousBodyTransformSimSpace.GetTranslation()) * 0.5f;
+				// Calculate constraint offset positions in the space of each body. For world space we scale the limit distances in UpdateLimits to account for dynamic scale. Store the unscaled value in ConstraintOffsets.
+				const FVector ConstaintLocationSimSpace = (BodyTransform.GetTranslation() - PreviousBodyTransformSimSpace.GetTranslation()) * 0.5f / (bScalingSupportEnabled ? PhysicsBody.Pose.Scale : 1);
 				ConstraintOffsets.Add(FAnimConstraintOffsetPair(PreviousBodyTransformSimSpace.GetRotation().UnrotateVector(ConstaintLocationSimSpace), BodyTransform.GetRotation().UnrotateVector(-ConstaintLocationSimSpace)));
 			}
 			else
@@ -1207,6 +1210,7 @@ void FAnimNode_AnimDynamics::UpdateLimits(FComponentSpacePoseContext& Output)
 	const FBoneContainer& BoneContainer = Output.Pose.GetPose().GetBoneContainer();
 
 	const bool bShouldScale = SimulationSpace == AnimPhysSimSpaceType::World && !CVarDisableWorldSpace.GetValueOnAnyThread() && CVarEnableScaling.GetValueOnAnyThread();
+	const bool bScaleSprings = bShouldScale && CVarScaleWorldSpaceSprings.GetValueOnAnyThread();
 
 	for(int32 ActiveIndex : ActiveBoneIndices)
 	{
@@ -1338,8 +1342,8 @@ void FAnimNode_AnimDynamics::UpdateLimits(FComponentSpacePoseContext& Output)
 		{
 			FAnimPhys::CreateSpring(Springs, PrevBody, ShapeTransform.GetTranslation(), &RigidBody, FVector::ZeroVector);
 			FAnimPhysSpring& NewSpring = Springs.Last();
-			NewSpring.SpringConstantLinear = LinearSpringConstant;
-			NewSpring.SpringConstantAngular = AngularSpringConstant;
+			NewSpring.SpringConstantLinear = LinearSpringConstant * (bScaleSprings ? RigidBody.Pose.Scale : 1);
+			NewSpring.SpringConstantAngular = AngularSpringConstant * (bScaleSprings ? RigidBody.Pose.Scale * RigidBody.Pose.Scale * RigidBody.Pose.Scale : 1);
 			NewSpring.AngularTarget = PhysicsBodyDef.ConstraintSetup.AngularTarget.GetSafeNormal();
 			NewSpring.AngularTargetAxis = PhysicsBodyDef.ConstraintSetup.AngularTargetAxis;
 			NewSpring.TargetOrientationOffset = ShapeTransform.GetRotation();
diff --git a/Engine/Source/Runtime/Engine/Private/Animation/AnimPhysicsSolver.cpp b/Engine/Source/Runtime/Engine/Private/Animation/AnimPhysicsSolver.cpp
--- a/Engine/Source/Runtime/Engine/Private/Animation/AnimPhysicsSolver.cpp
+++ b/Engine/Source/Runtime/Engine/Private/Animation/AnimPhysicsSolver.cpp
@@ -490,6 +490,16 @@ void FAnimPhys::ScaleRigidBodyMass(FAnimPhysRigidBody* InOutRigidBody, float Sca
 	InOutRigidBody->InertiaTensor *= Scale;
 }
 
+void FAnimPhys::RigidBodyShapeScaleUpdate(FAnimPhysRigidBody* InOutRigidBody, float Scale)
+{
+	const float InertialScaleFactor = Scale * Scale;
+
+	InOutRigidBody->AngularMomentum *= InertialScaleFactor;
+	InOutRigidBody->InverseWorldSpaceTensor *= 1.0f / InertialScaleFactor;
+	InOutRigidBody->InertiaTensor *= InertialScaleFactor;
+	InOutRigidBody->InverseTensorWithoutMass *= 1.0f / InertialScaleFactor;
+}
+
 FQuat FAnimPhys::DiffQ(const FQuat& InOrientation, const FMatrix &InInverseTensor, const FVector& InAngularMomentum)
 {
 	FQuat NormalisedOrient = InOrientation.GetNormalized();
diff --git a/Engine/Source/Runtime/Engine/Public/Animation/AnimPhysicsSolver.h b/Engine/Source/Runtime/Engine/Public/Animation/AnimPhysicsSolver.h
--- a/Engine/Source/Runtime/Engine/Public/Animation/AnimPhysicsSolver.h
+++ b/Engine/Source/Runtime/Engine/Public/Animation/AnimPhysicsSolver.h
@@ -481,6 +481,13 @@ public:
 	 */
 	static ENGINE_API void ScaleRigidBodyMass(FAnimPhysRigidBody* InOutRigidBody, float Scale);  // scales the mass and all relevant inertial properties by multiplier s
 
+	/** Update the inertial properties based on a scale change to the underlying shapes of the bodies independent of mass scaling
+	 *  @param InOutRigidBody Body to modify
+	 *  @param Scale Scale to use
+	 */
+	static ENGINE_API void RigidBodyShapeScaleUpdate(FAnimPhysRigidBody* InOutRigidBody, float Scale);
+
+
 	/** Apply an impulse to a body
 	 *  @param InOutRigidBody Body to apply the impulse to
 	 *  @param InWorldOrientedImpactPoint Location of the impulse
