Optimize "Remove Curve" anim node
- The allocation in Initialize_AnyThread was always allocating 0 so it was removed
- Don't compute / allocate MorphTargetCurveNames & MaterialCurveNames if we won't use them for filtering
- Faster execution path if we want to clear all curves with no filter
- Limit calls to TArray::Contains in CurveIsOfAllowedToRemoveType
- Re-arrange filter conditions to put less expensive check first
- Remove early exit in CacheBones_AnyThread since bEnabled is bindable it can change at runtime.

diff --git a/Engine/Source/Runtime/AnimGraphRuntime/Private/AnimNodes/AnimNode_RemoveCurve.cpp b/Engine/Source/Runtime/AnimGraphRuntime/Private/AnimNodes/AnimNode_RemoveCurve.cpp
--- a/Engine/Source/Runtime/AnimGraphRuntime/Private/AnimNodes/AnimNode_RemoveCurve.cpp
+++ b/Engine/Source/Runtime/AnimGraphRuntime/Private/AnimNodes/AnimNode_RemoveCurve.cpp
@@ -13,7 +13,6 @@ FAnimNode_RemoveCurve::FAnimNode_RemoveCurve()
 	bAffectMorphTargetCurves = true;
 	bAffectMaterialCurves = true;
 	bAffectRegularCurves = true;
-	NumCurvesOnSkeleton = 0;
 }
 
 void FAnimNode_RemoveCurve::Initialize_AnyThread(const FAnimationInitializeContext& Context)
@@ -21,27 +20,6 @@ void FAnimNode_RemoveCurve::Initialize_AnyThread(const FAnimationInitializeConte
 	DECLARE_SCOPE_HIERARCHICAL_COUNTER_ANIMNODE(Initialize_AnyThread)
 	Super::Initialize_AnyThread(Context);
 	SourcePose.Initialize(Context);
-
-	if (!bEnabled)
-	{
-		return;
-	}
-
-	// reserve expected amount
-	int ExcludedCurvesAmount = 0;
-	const bool bRemoveAllCurves = RemoveMode == ERemoveCurveMode::KeepOnlySpecifiedCurves ? true : false;
-
-	if (bRemoveAllCurves)
-	{
-		// technically, it's possible that pose has more curves than those that are specified on skeleton, but afaik unreal will allocate more automaticaally if that happens.
-		ExcludedCurvesAmount = NumCurvesOnSkeleton;
-	}
-	else
-	{
-		ExcludedCurvesAmount = Curves.Num();
-	}
-	
-	CurvesForRemoval.Reserve(ExcludedCurvesAmount);
 }
 
 void FAnimNode_RemoveCurve::CacheBones_AnyThread(const FAnimationCacheBonesContext& Context)
@@ -50,11 +28,6 @@ void FAnimNode_RemoveCurve::CacheBones_AnyThread(const FAnimationCacheBonesConte
 	Super::CacheBones_AnyThread(Context);
 	SourcePose.CacheBones(Context);
 
-	if (!bEnabled)
-	{
-		return;
-	}
-	
 	// Store curve types on the skel to look up in Eval state
 	MorphTargetCurveNames.Reset();
 	MaterialCurveNames.Reset();
@@ -66,25 +39,29 @@ void FAnimNode_RemoveCurve::CacheBones_AnyThread(const FAnimationCacheBonesConte
 		return;
 	}
 	
-	TArray<FName> MetaNames;
-	Skel->GetCurveMetaDataNames(MetaNames);
-	for (const FName& N : MetaNames)
+	const bool bAffectsAllCurves = bAffectMorphTargetCurves && bAffectMaterialCurves && bAffectRegularCurves;
+	// No need to fill curves array if we filter regardless of type
+	if (!bAffectsAllCurves)
 	{
-		if (const FCurveMetaData* Meta = Skel->GetCurveMetaData(N))
+		TArray<FName> MetaNames;
+		Skel->GetCurveMetaDataNames(MetaNames);
+	
+		for (const FName& N : MetaNames)
 		{
-			if (Meta->Type.bMorphtarget)
+			if (const FCurveMetaData* Meta = Skel->GetCurveMetaData(N))
 			{
-				MorphTargetCurveNames.Add(N);
+				if (Meta->Type.bMorphtarget)
+				{
+					MorphTargetCurveNames.Add(N);
+				}
+
+				if (Meta->Type.bMaterial)
+				{
+					MaterialCurveNames.Add(N);
+				}
 			}
-
-			if (Meta->Type.bMaterial)
-			{
-				MaterialCurveNames.Add(N);
-			}			
 		}
 	}
-
-	NumCurvesOnSkeleton = MetaNames.Num();
 }
 
 void FAnimNode_RemoveCurve::Update_AnyThread(const FAnimationUpdateContext& Context)
@@ -92,7 +69,6 @@ void FAnimNode_RemoveCurve::Update_AnyThread(const FAnimationUpdateContext& Cont
 	DECLARE_SCOPE_HIERARCHICAL_COUNTER_ANIMNODE(Update_AnyThread)
 	SourcePose.Update(Context);
 	GetEvaluateGraphExposedInputs().Execute(Context);
-
 }
 
 void FAnimNode_RemoveCurve::Evaluate_AnyThread(FPoseContext& Output)
@@ -114,6 +90,14 @@ void FAnimNode_RemoveCurve::Evaluate_AnyThread(FPoseContext& Output)
 	{
 		return;
 	}
+
+	// If we want to remove all curves, just do that directly
+	const bool bAffectsAllCurves = bAffectMorphTargetCurves && bAffectMaterialCurves && bAffectRegularCurves;
+	if (bRemoveAllCurves && bAffectsAllCurves && Curves.IsEmpty())
+	{
+		Output.Curve.Empty();
+		return;
+	}
 	
 	CurvesForRemoval = {};
 
@@ -121,29 +105,39 @@ void FAnimNode_RemoveCurve::Evaluate_AnyThread(FPoseContext& Output)
 	auto CurveIsOfAllowedToRemoveType = [this](const FName& CurveName) -> bool
 	{
 		// not allowed to be removed cause the type is checked off on the node
-		bool bAllowedIfMorphTargetCurve = bAffectMorphTargetCurves ? true : !MorphTargetCurveNames.Contains(CurveName);
-		bool bAllowedIfMaterialCurve = bAffectMaterialCurves ? true : !MaterialCurveNames.Contains(CurveName);
-		bool bAllowedIfRegularCurve = bAffectRegularCurves ? true : (MorphTargetCurveNames.Contains(CurveName) || MaterialCurveNames.Contains(CurveName));
+		if (!bAffectRegularCurves)
+		{
+			const bool bIsMorphTargetCurve = MorphTargetCurveNames.Contains(CurveName);
+			const bool bIsMaterialCurve = MaterialCurveNames.Contains(CurveName);
 
-		return bAllowedIfMorphTargetCurve && bAllowedIfMaterialCurve && bAllowedIfRegularCurve;
+			return (bIsMorphTargetCurve || bIsMaterialCurve) &&			// Regular
+				(bAffectMorphTargetCurves || !bIsMorphTargetCurve) &&	// MorphTarget
+				(bAffectMaterialCurves || !bIsMaterialCurve);			// Material
+
+		}
+		else
+		{
+			return (bAffectMorphTargetCurves || !MorphTargetCurveNames.Contains(CurveName)) &&	// Morph target
+				(bAffectMaterialCurves || !MaterialCurveNames.Contains(CurveName));				// Material
+		}
 	};
 
 	// If Removal of all curves is requested
 	if (bRemoveAllCurves)
 	{
 		CurvesForRemoval.Reserve(Output.Curve.Num());
-		
-		Output.Curve.ForEachElement([this, CurveIsOfAllowedToRemoveType](const UE::Anim::FCurveElement& Elem)
-		{	
-			
-			// is in exceptions list
-			bool bIsInExceptions = Curves.Contains(Elem.Name);
-
-			if (CurveIsOfAllowedToRemoveType(Elem.Name) && !bIsInExceptions)
+
+		Output.Curve.ForEachElement([this, bAffectsAllCurves, CurveIsOfAllowedToRemoveType](const UE::Anim::FCurveElement& Elem)
 			{
-				CurvesForRemoval.Add(Elem.Name);
-			}
-		});
+
+				// is in exceptions list
+				bool bIsInExceptions = Curves.Contains(Elem.Name);
+
+				if (!bIsInExceptions && (bAffectsAllCurves || CurveIsOfAllowedToRemoveType(Elem.Name)))
+				{
+					CurvesForRemoval.Add(Elem.Name);
+				}
+			});
 	}
 	else
 	{
@@ -151,25 +145,23 @@ void FAnimNode_RemoveCurve::Evaluate_AnyThread(FPoseContext& Output)
 
 		// iterate through user provided list of curve names and remove
 		for (const FName& CurveName : Curves)
-		{	
-			
-			if (CurveIsOfAllowedToRemoveType(CurveName))
+		{
+			if (bAffectsAllCurves || CurveIsOfAllowedToRemoveType(CurveName))
 			{
 				CurvesForRemoval.Add(CurveName);
 			}
-		}			
+		};
 	}
 
-	// Remove each collected and aproved curve from this pose's curve container	
+	// Remove each collected and approved curve from this pose's curve container	
 	UE::Anim::FNamedValueArrayUtils::RemoveByPredicate(
 		Output.Curve,
 		CurvesForRemoval,
 		[](const UE::Anim::FCurveElement& OutputCurve, const UE::Anim::FCurveElement& CurveToRemove)
-		{			
+		{
 			return true;
 		}
 	);
-
 }
 
 void FAnimNode_RemoveCurve::GatherDebugData(FNodeDebugData& DebugData)
@@ -204,4 +196,4 @@ void FAnimNode_RemoveCurve::GatherDebugData(FNodeDebugData& DebugData)
 	DebugData.AddDebugItem(DebugLine);
 
 	SourcePose.GatherDebugData(DebugData);
-}
\ No newline at end of file
+}
diff --git a/Engine/Source/Runtime/AnimGraphRuntime/Public/AnimNodes/AnimNode_RemoveCurve.h b/Engine/Source/Runtime/AnimGraphRuntime/Public/AnimNodes/AnimNode_RemoveCurve.h
--- a/Engine/Source/Runtime/AnimGraphRuntime/Public/AnimNodes/AnimNode_RemoveCurve.h
+++ b/Engine/Source/Runtime/AnimGraphRuntime/Public/AnimNodes/AnimNode_RemoveCurve.h
@@ -66,9 +66,6 @@ public:
 	UPROPERTY(transient)
 	TArray<FName> MaterialCurveNames;
 
-	UPROPERTY(transient)
-	int32 NumCurvesOnSkeleton;
-
 	ANIMGRAPHRUNTIME_API FAnimNode_RemoveCurve();
 
 	// FAnimNode_Base interface
diff --git a/Engine/Source/Runtime/Engine/Public/Animation/NamedValueArray.h b/Engine/Source/Runtime/Engine/Public/Animation/NamedValueArray.h
--- a/Engine/Source/Runtime/Engine/Public/Animation/NamedValueArray.h
+++ b/Engine/Source/Runtime/Engine/Public/Animation/NamedValueArray.h
@@ -770,7 +770,7 @@ struct FNamedValueArrayUtils
 	}
 
 	/**
-	 * Removes elements in InOutValueArray0 that match InValueArray1 if predicate returns false
+	 * Removes elements in InOutValueArray0 that match InValueArray1 if predicate returns true
 	 **/
 	template<typename AllocatorType0, typename ElementType0, typename AllocatorType1, typename ElementType1, typename PredicateType>
 	static void RemoveByPredicate(TNamedValueArray<AllocatorType0, ElementType0>& InOutValueArray0, const TNamedValueArray<AllocatorType1, ElementType1>& InValueArray1, PredicateType InPredicate)
