Expose bIgnorePoseAssetSourceAnimationCurves as a pin on the pose blend node

diff --git a/Engine/Source/Runtime/AnimGraphRuntime/Public/AnimNodes/AnimNode_PoseBlendNode.h b/Engine/Source/Runtime/AnimGraphRuntime/Public/AnimNodes/AnimNode_PoseBlendNode.h
--- a/Engine/Source/Runtime/AnimGraphRuntime/Public/AnimNodes/AnimNode_PoseBlendNode.h
+++ b/Engine/Source/Runtime/AnimGraphRuntime/Public/AnimNodes/AnimNode_PoseBlendNode.h
@@ -28,7 +28,7 @@ public:
 	UPROPERTY(EditAnywhere, Category = "Blend")
 	TObjectPtr<UCurveFloat> CustomCurve;
 
-	UPROPERTY(EditAnywhere, Category = "Blend")
+	UPROPERTY(EditAnywhere, Category = "Blend", BlueprintReadWrite, meta = (PinHiddenByDefault))
 	bool bIgnorePoseAssetSourceAnimationCurves = false;
 
 private:
