[UE] AnimDynamics - Ensure overridden sim space settings are passed into simulation
- Pass UseSimulationSettings into CalculateSimulationSpaceMotion()
- Always read from Settings argument instead of data member in CalculateSimulationSpaceMotion()

#rnx

diff --git a/Engine/Source/Runtime/AnimGraphRuntime/Private/BoneControllers/AnimNode_AnimDynamics.cpp b/Engine/Source/Runtime/AnimGraphRuntime/Private/BoneControllers/AnimNode_AnimDynamics.cpp
--- a/Engine/Source/Runtime/AnimGraphRuntime/Private/BoneControllers/AnimNode_AnimDynamics.cpp
+++ b/Engine/Source/Runtime/AnimGraphRuntime/Private/BoneControllers/AnimNode_AnimDynamics.cpp
@@ -368,7 +368,7 @@ void FAnimNode_AnimDynamics::CalculateSimulationSpaceMotion(
 	// If the simulation is in world space, the simulation space is stationary
 	if (Space == AnimPhysSimSpaceType::World && !CVarDisableWorldSpace.GetValueOnAnyThread())
 	{
-		SpaceAngularVel = SimSpaceSettings.ExternalAngularVelocity;
+		SpaceAngularVel = Settings.ExternalAngularVelocity;
 		return;
 	}
 
@@ -589,7 +589,7 @@ void FAnimNode_AnimDynamics::EvaluateSkeletalControl_AnyThread(FComponentSpacePo
 				}
 			}
 
-			CalculateSimulationSpaceMotion(SimulationSpace, SimSpaceSettings, Output, NextTimeStep, SimFrameAngularVelocity, SimFrameAngularAcc);
+			CalculateSimulationSpaceMotion(SimulationSpace, UseSimSpaceSettings, Output, NextTimeStep, SimFrameAngularVelocity, SimFrameAngularAcc);
 
 			if (SimulationSpace != AnimPhysSimSpaceType::World || CVarDisableWorldSpace.GetValueOnAnyThread())
 			{
