Add standard blend for mirroring

diff --git a/Engine/Source/Runtime/AnimGraphRuntime/Private/AnimNodes/AnimNode_Mirror.cpp b/Engine/Source/Runtime/AnimGraphRuntime/Private/AnimNodes/AnimNode_Mirror.cpp
--- a/Engine/Source/Runtime/AnimGraphRuntime/Private/AnimNodes/AnimNode_Mirror.cpp
+++ b/Engine/Source/Runtime/AnimGraphRuntime/Private/AnimNodes/AnimNode_Mirror.cpp
@@ -32,6 +32,11 @@ bool FAnimNode_MirrorBase::SetMirrorDataTable(UMirrorDataTable* MirrorTable)
 	return false;
 }
 
+EMirrorBlendType FAnimNode_MirrorBase::GetBlendType() const
+{
+	return EMirrorBlendType::Inertialization;
+}
+
 bool FAnimNode_MirrorBase::GetMirror() const
 {
 	return false;
@@ -106,6 +111,9 @@ void FAnimNode_MirrorBase::Initialize_AnyThread(const FAnimationInitializeContex
 {
 	FAnimNode_Base::Initialize_AnyThread(Context);
 	Source.Initialize(Context);
+
+	// Invalidate mirror state to avoid triggering blends on first frame after a re-init.
+	bMirrorStateIsValid = false;
 }
 
 void FAnimNode_MirrorBase::CacheBones_AnyThread(const FAnimationCacheBonesContext& Context)
@@ -117,6 +125,47 @@ void FAnimNode_MirrorBase::CacheBones_AnyThread(const FAnimationCacheBonesContex
 	FillCompactPoseAndComponentRefRotations(BoneContainer);
 }
 
+void FAnimNode_MirrorBase::MirrorPoseInPlace(FPoseContext& InPose)
+{
+	UMirrorDataTable* MirrorDataTable = GetMirrorDataTable();
+	if (MirrorDataTable)
+	{
+		if (GetBoneMirroring())
+		{
+			if (InPose.ExpectsAdditivePose())
+			{
+				FText Message = FText::Format(LOCTEXT("AdditiveMirrorWarning", "Trying to mirror an additive animation pose is not supported in anim instance '{0}'"), FText::FromString(InPose.AnimInstanceProxy->GetAnimInstanceName()));
+				InPose.LogMessage(EMessageSeverity::Warning, Message);
+
+				// Force a bind pose to make it obvious
+				InPose.Pose.ResetToAdditiveIdentity();
+			}
+			else
+			{
+				const FBoneContainer& BoneContainer = InPose.Pose.GetBoneContainer();
+				const TArray<FBoneIndexType>& RequiredBoneIndices = BoneContainer.GetBoneIndicesArray();
+				int32 NumReqBones = RequiredBoneIndices.Num();
+				if (CompactPoseMirrorBones.Num() != NumReqBones)
+				{
+					FillCompactPoseAndComponentRefRotations(BoneContainer);
+				}
+
+				FAnimationRuntime::MirrorPose(InPose.Pose, MirrorDataTable->MirrorAxis, CompactPoseMirrorBones, ComponentSpaceRefRotations);
+			}
+		}
+
+		if (GetCurveMirroring())
+		{
+			FAnimationRuntime::MirrorCurves(InPose.Curve, *MirrorDataTable);
+		}
+
+		if (GetAttributeMirroring())
+		{
+			UE::Anim::Attributes::MirrorAttributes(InPose.CustomAttributes, *MirrorDataTable, CompactPoseMirrorBones);
+		}
+	}
+}
+
 void FAnimNode_MirrorBase::FillCompactPoseAndComponentRefRotations(const FBoneContainer& BoneContainer)
 {
 	UMirrorDataTable* MirrorDataTable = GetMirrorDataTable();
@@ -141,39 +190,60 @@ void FAnimNode_MirrorBase::Update_AnyThread(const FAnimationUpdateContext& Conte
 	bool bMirror = GetMirror();
 	bool bRequestedInertialization = false;
 	
-	if (bMirrorStateIsValid && bMirrorState != bMirror)
+	if (bMirrorStateIsValid)
 	{
-		if (GetBlendTimeOnMirrorStateChange() > SMALL_NUMBER)
+		if (bMirrorState != bMirror)
 		{
-			// Inertialize when switching between mirrored and unmirrored states to smooth out the pose discontinuity
-			UE::Anim::IInertializationRequester* InertializationRequester = Context.GetMessage<UE::Anim::IInertializationRequester>();
-			if (InertializationRequester)
+			Blend.SetDesiredValue(bMirror ? 1.0f : 0.0f);
+			if ((GetBlendType() == EMirrorBlendType::Inertialization))
 			{
-				FInertializationRequest Request;
-				Request.Duration = GetBlendTimeOnMirrorStateChange();
+				Blend.SetBlendTime(0.0f);
+				if (GetBlendTimeOnMirrorStateChange() > SMALL_NUMBER)
+				{
+					// Inertialize when switching between mirrored and unmirrored states to smooth out the pose discontinuity
+					UE::Anim::IInertializationRequester* InertializationRequester = Context.GetMessage<UE::Anim::IInertializationRequester>();
+					if (InertializationRequester)
+					{
+						FInertializationRequest Request;
+						Request.Duration = GetBlendTimeOnMirrorStateChange();
 #if ANIM_TRACE_ENABLED
-				Request.NodeId = Context.GetCurrentNodeId();
-				Request.AnimInstance = Context.AnimInstanceProxy->GetAnimInstanceObject();
+						Request.NodeId = Context.GetCurrentNodeId();
+						Request.AnimInstance = Context.AnimInstanceProxy->GetAnimInstanceObject();
 #endif
 
-				InertializationRequester->RequestInertialization(Request);
-				InertializationRequester->AddDebugRecord(*Context.AnimInstanceProxy, Context.GetCurrentNodeId());
-				
-				bRequestedInertialization = true;
+						InertializationRequester->RequestInertialization(Request);
+						InertializationRequester->AddDebugRecord(*Context.AnimInstanceProxy, Context.GetCurrentNodeId());
+
+						bRequestedInertialization = true;
+					}
+					else
+					{
+						FAnimNode_Inertialization::LogRequestError(Context, Source);
+					}
+				}
 			}
 			else
 			{
-				FAnimNode_Inertialization::LogRequestError(Context, Source);
+				Blend.SetBlendTime(GetBlendTimeOnMirrorStateChange());
 			}
-		}
 
-		// Optionally reinitialize the source when the mirror state changes
-		if (GetResetChildOnMirrorStateChange())
-		{
-			FAnimationInitializeContext ReinitializeContext(Context.AnimInstanceProxy, Context.SharedContext);
-			Source.Initialize(ReinitializeContext);
+			// Optionally reinitialize the source when the mirror state changes
+			if (GetResetChildOnMirrorStateChange())
+			{
+				FAnimationInitializeContext ReinitializeContext(Context.AnimInstanceProxy, Context.SharedContext);
+				Source.Initialize(ReinitializeContext);
+			}
 		}
+
+		Blend.Update(Context.GetDeltaTime());
+	}
+	else
+	{
+		const float DesiredAlpha = bMirror ? 1.0f : 0.0f;
+		Blend.SetAlpha(DesiredAlpha);
+		Blend.SetDesiredValue(DesiredAlpha);
 	}
+
 	UMirrorDataTable* MirrorDataTable = GetMirrorDataTable();
 	UE::Anim::TOptionalScopedGraphMessage<UE::Anim::FMirrorSyncScope> Message(bMirror, Context, Context, MirrorDataTable);
 	UE::Anim::TOptionalScopedGraphMessage<UE::Anim::FAnimInertializationSyncScope> InertializationSync(bRequestedInertialization, Context);
@@ -188,44 +258,32 @@ void FAnimNode_MirrorBase::Update_AnyThread(const FAnimationUpdateContext& Conte
 
 void FAnimNode_MirrorBase::Evaluate_AnyThread(FPoseContext& Output)
 {
-	Source.Evaluate(Output);
-	UMirrorDataTable* MirrorDataTable = GetMirrorDataTable();
 
-	if (bMirrorState && MirrorDataTable)
+	const float BlendAlpha = Blend.GetBlendedValue();
+	if (FAnimWeight::IsFullWeight(BlendAlpha))
 	{
-		if (GetBoneMirroring())
-		{
-			if (Output.ExpectsAdditivePose())
-			{
-				FText Message = FText::Format(LOCTEXT("AdditiveMirrorWarning", "Trying to mirror an additive animation pose is not supported in anim instance '{0}'"), FText::FromString(Output.AnimInstanceProxy->GetAnimInstanceName()));
-				Output.LogMessage(EMessageSeverity::Warning, Message);
-
-				// Force a bind pose to make it obvious
-				Output.Pose.ResetToAdditiveIdentity();
-			}
-			else
-			{
-				const FBoneContainer& BoneContainer = Output.Pose.GetBoneContainer();
-				const TArray<FBoneIndexType>& RequiredBoneIndices = BoneContainer.GetBoneIndicesArray();
-				int32 NumReqBones = RequiredBoneIndices.Num();
-				if (CompactPoseMirrorBones.Num() != NumReqBones)
-				{
-					FillCompactPoseAndComponentRefRotations(BoneContainer);
-				}
-
-				FAnimationRuntime::MirrorPose(Output.Pose, MirrorDataTable->MirrorAxis, CompactPoseMirrorBones, ComponentSpaceRefRotations);
-			}
-		}
-
-		if (GetCurveMirroring())
-		{
-			FAnimationRuntime::MirrorCurves(Output.Curve, *MirrorDataTable);
-		}
-
-		if (GetAttributeMirroring())
-		{
-			UE::Anim::Attributes::MirrorAttributes(Output.CustomAttributes, *MirrorDataTable, CompactPoseMirrorBones);
-		}
+		Source.Evaluate(Output);
+		MirrorPoseInPlace(Output);
+	}
+	else if (FAnimWeight::IsRelevant(BlendAlpha))
+	{
+		FPoseContext UnmirroredPose(Output);
+		Source.Evaluate(UnmirroredPose);
+
+		FPoseContext MirroredPose(Output);
+		MirroredPose = UnmirroredPose;
+		// @TODO: Make BlendXPosesTogether that can blend in-place to save 3rd pose + unneeded pose copy.
+		// This is a common pattern i.e. FAnimNode_TwoWayBlend.
+		MirrorPoseInPlace(MirroredPose);
+		
+		FAnimationPoseData OutAnimationPoseData(Output);
+		const FAnimationPoseData UnmirroredPoseData(UnmirroredPose);
+		const FAnimationPoseData MirroredPoseData(MirroredPose);
+		FAnimationRuntime::BlendTwoPosesTogether(UnmirroredPoseData, MirroredPoseData, (1.0f - BlendAlpha), OutAnimationPoseData);
+	}
+	else
+	{
+		Source.Evaluate(Output);
 	}
 }
 
@@ -262,6 +320,12 @@ bool FAnimNode_Mirror::SetMirrorDataTable(UMirrorDataTable* InMirrorTable)
 	return false;
 }
 
+EMirrorBlendType FAnimNode_Mirror::GetBlendType() const
+{
+	return GET_ANIM_NODE_DATA(EMirrorBlendType, BlendType);
+
+}
+
 bool FAnimNode_Mirror::GetMirror() const
 {
 	return GET_ANIM_NODE_DATA(bool, bMirror);
@@ -395,6 +459,11 @@ bool FAnimNode_Mirror_Standalone::SetMirrorDataTable(UMirrorDataTable* MirrorTab
 	return true;
 }
 
+EMirrorBlendType FAnimNode_Mirror_Standalone::GetBlendType() const
+{
+	return BlendType;
+}
+
 bool FAnimNode_Mirror_Standalone::GetMirror() const
 {
 	return bMirror;
diff --git a/Engine/Source/Runtime/AnimGraphRuntime/Public/AnimNodes/AnimNode_Mirror.h b/Engine/Source/Runtime/AnimGraphRuntime/Public/AnimNodes/AnimNode_Mirror.h
--- a/Engine/Source/Runtime/AnimGraphRuntime/Public/AnimNodes/AnimNode_Mirror.h
+++ b/Engine/Source/Runtime/AnimGraphRuntime/Public/AnimNodes/AnimNode_Mirror.h
@@ -6,8 +6,18 @@
 #include "UObject/ObjectMacros.h"
 #include "Animation/AnimNodeBase.h"
 #include "Animation/MirrorDataTable.h"
+#include "AlphaBlend.h"
 #include "AnimNode_Mirror.generated.h"
 
+UENUM()
+enum class EMirrorBlendType : uint8
+{
+	// Use a standard crossfade blend. If Reset Child on Mirror State Change is enabled, this will result in a pop.
+	StandardBlend,
+	// Use an inertial blend. This requires having a matching inertialization or dead blend node upstream in the graph.
+	Inertialization
+};
+
 USTRUCT(BlueprintInternalUseOnly)
 struct FAnimNode_MirrorBase : public FAnimNode_Base
 {
@@ -26,9 +36,11 @@ public:
 	// Get the MirrorDataTable
 	ANIMGRAPHRUNTIME_API virtual UMirrorDataTable* GetMirrorDataTable() const;
 
-	// Set the MirrorDataTable
 	ANIMGRAPHRUNTIME_API virtual bool SetMirrorDataTable(UMirrorDataTable* MirrorTable);
 
+	// Get mirror blend type. Default to inertialization for backwards compatibility.
+	ANIMGRAPHRUNTIME_API virtual EMirrorBlendType GetBlendType() const;
+
 	// Get Mirror State
 	ANIMGRAPHRUNTIME_API virtual bool GetMirror() const;
 	// How long to blend using inertialization when switching  mirrored state
@@ -79,9 +91,12 @@ protected:
 	FPoseLink Source;
 
 private:
+	FAlphaBlend Blend;
+
 	bool bMirrorState;
 	bool bMirrorStateIsValid;
 
+	void MirrorPoseInPlace(FPoseContext& InPose);
 	void FillCompactPoseAndComponentRefRotations(const FBoneContainer& BoneContainer);
 	// Compact pose format of Mirror Bone Map
 	TCustomBoneIndexArray<FCompactPoseBoneIndex, FCompactPoseBoneIndex> CompactPoseMirrorBones;
@@ -105,6 +120,7 @@ public:
 	ANIMGRAPHRUNTIME_API virtual UMirrorDataTable* GetMirrorDataTable() const override;
 	ANIMGRAPHRUNTIME_API virtual bool SetMirrorDataTable(UMirrorDataTable* MirrorTable) override;
 
+	ANIMGRAPHRUNTIME_API virtual EMirrorBlendType GetBlendType() const;
 	ANIMGRAPHRUNTIME_API virtual bool GetMirror() const override;
 	ANIMGRAPHRUNTIME_API virtual float GetBlendTimeOnMirrorStateChange() const override;
 	ANIMGRAPHRUNTIME_API virtual bool GetBoneMirroring() const override;
@@ -131,6 +147,10 @@ protected:
 	// Inertialization blend time to use when transitioning between mirrored and unmirrored states
 	UPROPERTY(EditAnywhere, Category = MirrorTransition, meta = (PinHiddenByDefault, FoldProperty))
 	float BlendTime = 0.0f;
+
+	UPROPERTY(EditAnywhere, Category = MirrorTransition, meta = (PinHiddenByDefault, FoldProperty))
+	EMirrorBlendType BlendType = EMirrorBlendType::Inertialization;
+
 	// Whether to reset (reinitialize) the child (source) pose when the mirror state changes
 	UPROPERTY(EditAnywhere, Category = MirrorTransition, meta = (FoldProperty))
 	bool bResetChild = false;
@@ -162,6 +182,7 @@ public:
 	ANIMGRAPHRUNTIME_API virtual UMirrorDataTable* GetMirrorDataTable() const override;
 	ANIMGRAPHRUNTIME_API virtual bool SetMirrorDataTable(UMirrorDataTable* MirrorTable) override;
 
+	ANIMGRAPHRUNTIME_API virtual EMirrorBlendType GetBlendType() const;
 	ANIMGRAPHRUNTIME_API virtual bool GetMirror() const override;
 	ANIMGRAPHRUNTIME_API virtual float GetBlendTimeOnMirrorStateChange() const override;
 	ANIMGRAPHRUNTIME_API virtual bool GetBoneMirroring() const override;
@@ -188,6 +209,9 @@ protected:
 	UPROPERTY(EditAnywhere, Category = MirrorTransition, meta = (PinHiddenByDefault, FoldProperty))
 	float BlendTime = 0.0f;
 
+	UPROPERTY(EditAnywhere, Category = MirrorTransition, meta = (PinHiddenByDefault, FoldProperty))
+	EMirrorBlendType BlendType = EMirrorBlendType::Inertialization;
+
 	// Whether to reset (reinitialize) the child (source) pose when the mirror state changes
 	UPROPERTY(EditAnywhere, Category = MirrorTransition, meta = (FoldProperty))
 	bool bResetChild = false;
