Move leader component test earlier in CopyPoseFromMesh to make sure our bone map is updated correctly

diff --git a/Engine/Source/Runtime/AnimGraphRuntime/Private/AnimNodes/AnimNode_CopyPoseFromMesh.cpp b/Engine/Source/Runtime/AnimGraphRuntime/Private/AnimNodes/AnimNode_CopyPoseFromMesh.cpp
--- a/Engine/Source/Runtime/AnimGraphRuntime/Private/AnimNodes/AnimNode_CopyPoseFromMesh.cpp
+++ b/Engine/Source/Runtime/AnimGraphRuntime/Private/AnimNodes/AnimNode_CopyPoseFromMesh.cpp
@@ -15,6 +15,15 @@
 /////////////////////////////////////////////////////
 // FAnimNode_CopyPoseFromMesh
 
+namespace
+{
+	static bool bTestLeaderPoseEarly = true;
+	static FAutoConsoleVariableRef CVarTestLeaderPoseEarly(
+		TEXT("a.CopyPose.TestLeaderPoseEarly"),
+		bTestLeaderPoseEarly,
+		TEXT("In CopyPoseFromMesh, test if we have a Leader Pose early to avoid potential mesh corruption"));
+}
+
 FAnimNode_CopyPoseFromMesh::FAnimNode_CopyPoseFromMesh()
 	: SourceMeshComponent(nullptr)
 	, bUseAttachedParent (false)
@@ -43,7 +52,17 @@ void FAnimNode_CopyPoseFromMesh::RefreshMeshComponent(USkeletalMeshComponent* Ta
 {
 	auto ResetMeshComponent = [this](USkeletalMeshComponent* InMeshComponent, USkeletalMeshComponent* InTargetMeshComponent)
 	{
+		// If our source is running under leader-pose, then get bone data from there
+		if (bTestLeaderPoseEarly && InMeshComponent)
+		{
+			if (USkeletalMeshComponent* LeaderPoseComponent = Cast<USkeletalMeshComponent>(InMeshComponent->LeaderPoseComponent.Get()))
+			{
+				InMeshComponent = LeaderPoseComponent;
+			}
+		}
+
 		USkeletalMeshComponent* CurrentMeshComponent = CurrentlyUsedSourceMeshComponent.Get();
+
 		// if current mesh exists, but not same as input mesh
 		if (CurrentMeshComponent)
 		{
@@ -122,10 +141,13 @@ void FAnimNode_CopyPoseFromMesh::PreUpdate(const UAnimInstance* InAnimInstance)
 
 	if (CurrentMeshComponent && CurrentMeshComponent->GetSkeletalMeshAsset() && CurrentMeshComponent->IsRegistered())
 	{
-		// If our source is running under leader-pose, then get bone data from there
-		if(USkeletalMeshComponent* LeaderPoseComponent = Cast<USkeletalMeshComponent>(CurrentMeshComponent->LeaderPoseComponent.Get()))
+		if (!bTestLeaderPoseEarly)
 		{
-			CurrentMeshComponent = LeaderPoseComponent;
+			// If our source is running under leader-pose, then get bone data from there
+			if (USkeletalMeshComponent* LeaderPoseComponent = Cast<USkeletalMeshComponent>(CurrentMeshComponent->LeaderPoseComponent.Get()))
+			{
+				CurrentMeshComponent = LeaderPoseComponent;
+			}
 		}
 
 		// re-check mesh component validity as it may have changed to leader
