AnimGraphRuntime - adding ability to update the base node first in FAnimNode_LayeredBoneBlend via newly introduced bUpdateBasePoseFirst property

#Synced-CL 49959264

diff --git a/Engine/Source/Runtime/AnimGraphRuntime/Private/AnimNodes/AnimNode_LayeredBoneBlend.cpp b/Engine/Source/Runtime/AnimGraphRuntime/Private/AnimNodes/AnimNode_LayeredBoneBlend.cpp
--- a/Engine/Source/Runtime/AnimGraphRuntime/Private/AnimNodes/AnimNode_LayeredBoneBlend.cpp
+++ b/Engine/Source/Runtime/AnimGraphRuntime/Private/AnimNodes/AnimNode_LayeredBoneBlend.cpp
@@ -154,10 +154,16 @@ void FAnimNode_LayeredBoneBlend::Update_AnyThread(const FAnimationUpdateContext&
 	float RootMotionWeight = 0.f;
 	const float RootMotionClearWeight = bBlendRootMotionBasedOnRootBone ? 0.f : 1.f;
 
+	TArray<float, TInlineAllocator<64>> ChildRootMotionWeight;
 	if (IsLODEnabled(Context.AnimInstanceProxy))
 	{
 		GetEvaluateGraphExposedInputs().Execute(Context);
 
+		if (bUpdateBasePoseFirst)
+		{
+			ChildRootMotionWeight.SetNumZeroed(BlendPoses.Num());
+		}
+
 		for (int32 ChildIndex = 0; ChildIndex < BlendPoses.Num(); ++ChildIndex)
 		{
 			const float ChildWeight = BlendWeights[ChildIndex];
@@ -184,7 +190,14 @@ void FAnimNode_LayeredBoneBlend::Update_AnyThread(const FAnimationUpdateContext&
 				}
 
 				const float ThisPoseRootMotionWeight = (ChildIndex == RootMotionBlendPose) ? RootMotionWeight : RootMotionClearWeight;
-				BlendPoses[ChildIndex].Update(Context.FractionalWeightAndRootMotion(ChildWeight, ThisPoseRootMotionWeight));
+				if (bUpdateBasePoseFirst)
+				{
+					ChildRootMotionWeight[ChildIndex] = ThisPoseRootMotionWeight;
+				}
+				else
+				{
+					BlendPoses[ChildIndex].Update(Context.FractionalWeightAndRootMotion(ChildWeight, ThisPoseRootMotionWeight));
+				}
 			}
 		}
 	}
@@ -201,6 +214,22 @@ void FAnimNode_LayeredBoneBlend::Update_AnyThread(const FAnimationUpdateContext&
 		BasePose.Update(Context);
 	}
 
+	if (!ChildRootMotionWeight.IsEmpty())
+	{
+		check(bUpdateBasePoseFirst);
+		check(ChildRootMotionWeight.Num() == BlendPoses.Num());
+
+		for (int32 ChildIndex = 0; ChildIndex < BlendPoses.Num(); ++ChildIndex)
+		{
+			const float ChildWeight = BlendWeights[ChildIndex];
+			if (FAnimWeight::IsRelevant(ChildWeight))
+			{
+				const float ThisPoseRootMotionWeight = ChildRootMotionWeight[ChildIndex];
+				BlendPoses[ChildIndex].Update(Context.FractionalWeightAndRootMotion(ChildWeight, ThisPoseRootMotionWeight));
+			}
+		}
+	}
+
 	TRACE_ANIM_NODE_VALUE(Context, TEXT("Num Poses"), BlendPoses.Num());
 }
 
diff --git a/Engine/Source/Runtime/AnimGraphRuntime/Public/AnimNodes/AnimNode_LayeredBoneBlend.h b/Engine/Source/Runtime/AnimGraphRuntime/Public/AnimNodes/AnimNode_LayeredBoneBlend.h
--- a/Engine/Source/Runtime/AnimGraphRuntime/Public/AnimNodes/AnimNode_LayeredBoneBlend.h
+++ b/Engine/Source/Runtime/AnimGraphRuntime/Public/AnimNodes/AnimNode_LayeredBoneBlend.h
@@ -109,6 +109,10 @@ public:
 	UPROPERTY(EditAnywhere, Category = Config)
 	bool bBlendRootMotionBasedOnRootBone;
 
+	/** Whether to update the base pose first or after all the blend poses */
+	UPROPERTY(Experimental, EditAnywhere, BlueprintReadWrite, Category = Config, meta = (PinHiddenByDefault))
+	bool bUpdateBasePoseFirst = false;
+
 	bool bHasRelevantPoses;
 
 	FAnimNode_LayeredBoneBlend()
