[UE] [RBAN] Apply bone transform scaling in the input pose to output bone transforms from physics simulation.
- Disabled by default, enable with CVar p.RigidBodyNode.EnablePoseScalePropagation
- Only applies scale to the output transforms if the input pose contains transforms with scale != 1
- Scale only affects the output pose, not the simulation. This means that scaled simulated parts can potentially clip through other body parts that have a different scale.

[Disabled Bughawk Validation. This CL did pass bughawk's tests (review 43280) but the request exceeded submit tool's timeout so it counted as a fail]

#Synced-CL 49988729

diff --git a/Engine/Source/Runtime/AnimGraphRuntime/Private/BoneControllers/AnimNode_RigidBody.cpp b/Engine/Source/Runtime/AnimGraphRuntime/Private/BoneControllers/AnimNode_RigidBody.cpp
--- a/Engine/Source/Runtime/AnimGraphRuntime/Private/BoneControllers/AnimNode_RigidBody.cpp
+++ b/Engine/Source/Runtime/AnimGraphRuntime/Private/BoneControllers/AnimNode_RigidBody.cpp
@@ -75,6 +75,8 @@ FAutoConsoleVariableRef CVarRigidBodyNodeWorldObjectExpiry(TEXT("p.RigidBodyNode
 bool bRBAN_IncludeClothColliders = true;
 FAutoConsoleVariableRef CVarRigidBodyNodeIncludeClothColliders(TEXT("p.RigidBodyNode.IncludeClothColliders"), bRBAN_IncludeClothColliders, TEXT("Include cloth colliders as kinematic bodies in the immediate physics simulation."), ECVF_Default);
 
+TAutoConsoleVariable<bool> CVarEnableRigidBodyNodePoseScalePropagation(TEXT("p.RigidBodyNode.EnablePoseScalePropagation"), false, TEXT("Runtime Enable/Disable applying bone transform scales to the output pose (after simulation)."), ECVF_Default);
+
 // FSimSpaceSettings forced overrides for testing
 bool bRBAN_SimSpace_EnableOverride = false;
 float RBAN_SimSpaceOverride_WorldAlpha = -1.0f;
@@ -694,6 +696,8 @@ void FAnimNode_RigidBody::EvaluateSkeletalControl_AnyThread(FComponentSpacePoseC
 		const bool bUseDeferredSimulationTask =
 			((SimulationTiming == ESimulationTiming::Default) || bRBAN_DeferredSimulationForceDefault) ? bRBAN_DeferredSimulationDefault : (SimulationTiming == ESimulationTiming::Deferred);
 
+		bool bApplyBoneTransformScaleToOutputPose = CVarEnableRigidBodyNodePoseScalePropagation.GetValueOnAnyThread();
+
 		FVector SimSpaceGravity(0.f);
 
 		// Only need to tick physics if we didn't reset and we have some time to simulate
@@ -767,19 +771,34 @@ void FAnimNode_RigidBody::EvaluateSkeletalControl_AnyThread(FComponentSpacePoseC
 				}
 			}
 
-			// @todo(ccaulfield): We should be interpolating kinematic targets for each sub-step below
-			for (const FOutputBoneData& OutputData : OutputBoneData)
 			{
-				const int32 BodyIndex = OutputData.BodyIndex;
-				if (!BodyAnimData[BodyIndex].bIsSimulated)
+				bool bInputPoseContainsScaledTransforms = false;
+
+				for (const FOutputBoneData& OutputData : OutputBoneData)
 				{
-					const FTransform& ComponentSpaceTM = Output.Pose.GetComponentSpaceTransform(OutputData.CompactPoseBoneIndex);
-					const FTransform BodyTM = ConvertCSTransformToSimSpace(SimulationSpace, ComponentSpaceTM, CompWorldSpaceTM, BaseBoneTM);
+					const int32 BodyIndex = OutputData.BodyIndex;
+					const bool bIsSimulated = BodyAnimData[BodyIndex].bIsSimulated;
+				
+					if (!bIsSimulated || bApplyBoneTransformScaleToOutputPose)
+					{
+						const FTransform& ComponentSpaceTM = Output.Pose.GetComponentSpaceTransform(OutputData.CompactPoseBoneIndex);
 
-					Bodies[BodyIndex]->SetKinematicTarget(BodyTM);
+						if(!bIsSimulated)
+						{
+							// @todo(ccaulfield): We should be interpolating kinematic targets for each sub-step below
+							const FTransform BodyTM = ConvertCSTransformToSimSpace(SimulationSpace, ComponentSpaceTM, CompWorldSpaceTM, BaseBoneTM);
+							Bodies[BodyIndex]->SetKinematicTarget(BodyTM);
+						}
+
+						// Check for scaling on input pose transforms.
+						bInputPoseContainsScaledTransforms |= !ComponentSpaceTM.GetScale3D().Equals(FVector(1.0f), SMALL_NUMBER);
+					}	
 				}
+
+				// Avoid unnecessary scaling of simulation results when applying them to the output pose if the scale is identity.
+				bApplyBoneTransformScaleToOutputPose &= bInputPoseContainsScaledTransforms;
 			}
-			
+
 			UpdateWorldForces(CompWorldSpaceTM, BaseBoneTM, DeltaSeconds);
 			SimSpaceGravity = WorldVectorToSpaceNoScale(SimulationSpace, WorldSpaceGravity, CompWorldSpaceTM, BaseBoneTM);
 			SimSpaceGravity *= RBAN_GravityScale;
@@ -850,6 +869,15 @@ void FAnimNode_RigidBody::EvaluateSkeletalControl_AnyThread(FComponentSpacePoseC
 			}
 		}
 		
+		// Local record of body transforms after the input pose transform scale has been applied.
+		// Only used if bApplyBoneTransformScaleToOutputPose is true.
+		TArray<FTransform> ScaledBodyTransforms; 
+
+		if(bApplyBoneTransformScaleToOutputPose)
+		{
+			ScaledBodyTransforms.Init(FTransform::Identity, Bodies.Num());
+		}
+
 		//write back to animation system
 		const FTransform& SimulationWorldSpaceTM = bUseDeferredSimulationTask ? PreviousCompWorldSpaceTM : CompWorldSpaceTM;
 		for (const FOutputBoneData& OutputData : OutputBoneData)
@@ -902,6 +930,29 @@ void FAnimNode_RigidBody::EvaluateSkeletalControl_AnyThread(FComponentSpacePoseC
 						}
 					}
 
+					// Apply scale in input pose to body transforms before writing to the output pose.
+					if (bApplyBoneTransformScaleToOutputPose)
+					{
+						const int32 ParentBodyIndex = OutputData.ParentBodyIndex;
+												
+						if (ParentBodyIndex != INDEX_NONE)
+						{	
+							const FTransform ParentTransform = Bodies[ParentBodyIndex]->GetWorldTransform();
+
+							if (!BodyAnimData[ParentBodyIndex].bIsSimulated)
+							{
+								const FCompactPoseBoneIndex& CompactParentBoneIndex = OutputData.BoneIndicesToParentBody.Last();
+								
+								ScaledBodyTransforms[ParentBodyIndex] = ParentTransform;
+								ScaledBodyTransforms[ParentBodyIndex].SetScale3D(Output.Pose.GetComponentSpaceTransform(CompactParentBoneIndex).GetScale3D());
+							}
+
+							BodyTM = BodyTM.GetRelativeTransform(ParentTransform) * ScaledBodyTransforms[ParentBodyIndex];
+						}
+
+						ScaledBodyTransforms[BodyIndex] = BodyTM;
+					}
+
 					FTransform ComponentSpaceTM;
 
 					switch(SimulationSpace)
