[UE] RBAN Bone Scaling avoid sometimes scaling twice and fix parent bone index
- If an explict bone transform scale is applied to the simulation's base bone it will be applied twice. Avoid this.
- Reference the simulated bodies parent bone, not the imediate child of that bone.

#Synced-CL 50481136

diff --git a/Engine/Source/Runtime/AnimGraphRuntime/Private/BoneControllers/AnimNode_RigidBody.cpp b/Engine/Source/Runtime/AnimGraphRuntime/Private/BoneControllers/AnimNode_RigidBody.cpp
--- a/Engine/Source/Runtime/AnimGraphRuntime/Private/BoneControllers/AnimNode_RigidBody.cpp
+++ b/Engine/Source/Runtime/AnimGraphRuntime/Private/BoneControllers/AnimNode_RigidBody.cpp
@@ -869,12 +869,17 @@ void FAnimNode_RigidBody::EvaluateSkeletalControl_AnyThread(FComponentSpacePoseC
 			}
 		}
 		
-		// Local record of body transforms after the input pose transform scale has been applied.
-		// Only used if bApplyBoneTransformScaleToOutputPose is true.
-		TArray<FTransform> ScaledBodyTransforms; 
+		FTransform SimulationBaseBoneTM = BaseBoneTM;
+		TArray<FTransform> ScaledBodyTransforms;
 
 		if(bApplyBoneTransformScaleToOutputPose)
 		{
+			// BodyTM will already include any scaling in the BaseBoneTM if bone transform scale was applied above. 
+			// Remove that scaling to avoid applying it twice.
+			SimulationBaseBoneTM.SetScale3D(FVector(1.0f));
+
+			// Local record of body transforms after the input pose transform scale has been applied.
+			// Only used if bApplyBoneTransformScaleToOutputPose is true.
 			ScaledBodyTransforms.Init(FTransform::Identity, Bodies.Num());
 		}
 
@@ -941,8 +946,9 @@ void FAnimNode_RigidBody::EvaluateSkeletalControl_AnyThread(FComponentSpacePoseC
 
 							if (!BodyAnimData[ParentBodyIndex].bIsSimulated)
 							{
-								const FCompactPoseBoneIndex& CompactParentBoneIndex = OutputData.BoneIndicesToParentBody.Last();
-								
+								// BoneIndicesToParentBody contains all indexs up to but not including the parent bone index. The last element should be a direct child of the parent.
+								const FCompactPoseBoneIndex& CompactParentBoneIndex = Output.AnimInstanceProxy->GetRequiredBones().GetParentBoneIndex(OutputData.BoneIndicesToParentBody.Last());
+	
 								ScaledBodyTransforms[ParentBodyIndex] = ParentTransform;
 								ScaledBodyTransforms[ParentBodyIndex].SetScale3D(Output.Pose.GetComponentSpaceTransform(CompactParentBoneIndex).GetScale3D());
 							}
@@ -959,7 +965,7 @@ void FAnimNode_RigidBody::EvaluateSkeletalControl_AnyThread(FComponentSpacePoseC
 					{
 						case ESimulationSpace::ComponentSpace: ComponentSpaceTM = BodyTM; break;
 						case ESimulationSpace::WorldSpace: ComponentSpaceTM = BodyTM.GetRelativeTransform(SimulationWorldSpaceTM); break;
-						case ESimulationSpace::BaseBoneSpace: ComponentSpaceTM = BodyTM * BaseBoneTM; break;
+						case ESimulationSpace::BaseBoneSpace: ComponentSpaceTM = BodyTM * SimulationBaseBoneTM; break;
 						default: ensureMsgf(false, TEXT("Unsupported Simulation Space")); ComponentSpaceTM = BodyTM;
 					}
 					
