[UE] AnimDynamics bone transform scaling support
- AD node already includes bone scale in the physics simulation (well, its a scaler rather than a vector quantity but thats enough for our current requirements) but the bone scale is not written to the output bone transforms. This change just overwrites the output bone transform scale with the scale value applied to the simulation.
- Also added an CVar gated feature to updaet the scale in the simulation if the input pose bone scale changes

#Synced-CL 51050143

diff --git a/Engine/Source/Runtime/AnimGraphRuntime/Private/BoneControllers/AnimNode_AnimDynamics.cpp b/Engine/Source/Runtime/AnimGraphRuntime/Private/BoneControllers/AnimNode_AnimDynamics.cpp
--- a/Engine/Source/Runtime/AnimGraphRuntime/Private/BoneControllers/AnimNode_AnimDynamics.cpp
+++ b/Engine/Source/Runtime/AnimGraphRuntime/Private/BoneControllers/AnimNode_AnimDynamics.cpp
@@ -32,6 +32,8 @@ TAutoConsoleVariable<bool>  CVarDisableWorldSpace(TEXT("p.AnimDynamics.DisableWo
 TAutoConsoleVariable<bool>  CVarScaleInitialBodyMassByComponentSpace(TEXT("p.AnimDynamics.ScaleInitialBodyMassByComponentSpace"), true, TEXT("Scale Initial Body Mass By Component Space Scale when using World Space."));
 TAutoConsoleVariable<bool>  CVarTeleportBodiesOnScaleChange(TEXT("p.AnimDynamics.TeleportBodiesOnScaleChange"), true, TEXT("Teleport Bodies On Scale Change when using World Space."));
 TAutoConsoleVariable<bool>  CVarScaleWorldSpaceSprings(TEXT("p.AnimDynamics.ScaleWorldSpaceSprings"), true, TEXT("Scale spring constants to match body scale changes when in world space."));
+TAutoConsoleVariable<bool>  CVarPropagateSimulationScaleToOutputPose(TEXT("p.AnimDynamics.EnableUniformPoseScalePropagation"), true, TEXT("Copies a the scale applied to the simulated bodies into the output bone transforms. This scale is the minimum value of the x,y and z scales that are set in the input bone transforms."));
+TAutoConsoleVariable<bool>  CVarUpdateBodiesOnScaleChange(TEXT("p.AnimDynamics.EnableUpdateBodiesOnScaleChange"), true, TEXT("When true, input bone transform scale is compared to the scale applied to simulated bodies in the last physics evaluation. Physics body scales will be updated in the next physics evaluation if they are different."));
 
 // FSimSpaceSettingsWithControl forced overrides for testing
 bool bAnimDynamics_SimSpace_EnableOverride = false;
@@ -228,6 +230,7 @@ FAnimNode_AnimDynamics::FAnimNode_AnimDynamics()
 , bDoPhysicsUpdateInEditor(true)
 #endif
 , bEnabled(true)
+, bShouldUpdateBodiesOnScaleChange(false)
 #if ENABLE_ANIM_DRAW_DEBUG
 , FilteredBoneIndex(INDEX_NONE)
 #endif
@@ -463,7 +466,7 @@ void FAnimNode_AnimDynamics::EvaluateSkeletalControl_AnyThread(FComponentSpacePo
 
 		if (bScalingEnabled)
 		{
-			if (!SimWorldTransform.GetScale3D().Equals(PreviousSimSpaceTM.GetScale3D(), UE_KINDA_SMALL_NUMBER))
+			if (bShouldUpdateBodiesOnScaleChange || !SimWorldTransform.GetScale3D().Equals(PreviousSimSpaceTM.GetScale3D(), UE_KINDA_SMALL_NUMBER))
 			{
 				const bool bWorldSpace = SimulationSpace == AnimPhysSimSpaceType::World && !CVarDisableWorldSpace.GetValueOnAnyThread();
 				const bool bTeleportEnabled = bWorldSpace && CVarTeleportBodiesOnScaleChange.GetValueOnAnyThread();
@@ -514,6 +517,8 @@ void FAnimNode_AnimDynamics::EvaluateSkeletalControl_AnyThread(FComponentSpacePo
 						}
 					}
 				}
+
+				bShouldUpdateBodiesOnScaleChange = false;
 			}
 		}
 
@@ -703,6 +708,8 @@ void FAnimNode_AnimDynamics::EvaluateSkeletalControl_AnyThread(FComponentSpacePo
 		{
 			CONDITIONAL_SCOPE_CYCLE_COUNTER(STAT_AnimDynamicsBoneEval, FAnimPhys::bEnableDetailedStats);
 
+			const bool bPropagateSimulationScaleToOutputPose = CVarPropagateSimulationScaleToOutputPose.GetValueOnAnyThread();
+
 			const FBoneContainer& BoneContainer = Output.Pose.GetPose().GetBoneContainer();
 
 			check(Bodies.Num() <= PhysicsBodyDefinitions.Num());
@@ -762,6 +769,12 @@ void FAnimNode_AnimDynamics::EvaluateSkeletalControl_AnyThread(FComponentSpacePo
 					NewBoneTransform.SetRotation(RetargetedRotation);
 				}
 
+				if (bPropagateSimulationScaleToOutputPose)
+				{
+					// Override output bone transform scale with a uniform scale vector with the same value on each axis as the scale value applied in the simulation.
+					NewBoneTransform.SetScale3D(FVector(CurrentBody.Pose.Scale));
+				}
+
 				NewBoneTransform = GetComponentSpaceTransformFromSimSpace(SimulationSpace, Output, NewBoneTransform);
 
 				OutBoneTransforms.Add(FBoneTransform(BoneIndex, NewBoneTransform));
@@ -1203,6 +1216,8 @@ void FAnimNode_AnimDynamics::UpdateLimits(FComponentSpacePoseContext& Output)
 {
 	CONDITIONAL_SCOPE_CYCLE_COUNTER(STAT_AnimDynamicsLimitUpdate, FAnimPhys::bEnableDetailedStats);
 
+	const bool bShouldTestForBoneTransformScaleChange = CVarUpdateBodiesOnScaleChange.GetValueOnAnyThread();
+
 	// We're always going to use the same number so don't realloc
 	LinearLimits.Empty(LinearLimits.Num());
 	AngularLimits.Empty(AngularLimits.Num());
@@ -1235,7 +1250,15 @@ void FAnimNode_AnimDynamics::UpdateLimits(FComponentSpacePoseContext& Output)
 
 		// Get joint transform
 		FCompactPoseBoneIndex BoneIndex = CurrentBoneRef.GetCompactPoseIndex(BoneContainer);
-		FTransform BoundBoneTransform = GetBoneTransformInSimSpace(Output, BoneIndex);
+
+		FTransform BoundBoneComponentSpaceTransform = Output.Pose.GetComponentSpaceTransform(BoneIndex);
+		FTransform BoundBoneTransform = GetSimSpaceTransformFromComponentSpace(SimulationSpace, Output, BoundBoneComponentSpaceTransform);
+
+		if(bShouldTestForBoneTransformScaleChange && !bShouldUpdateBodiesOnScaleChange)
+		{
+			// Check if input bone transform scale has changed since the last physics init.
+			bShouldUpdateBodiesOnScaleChange = RigidBody.Pose.Scale != BoundBoneComponentSpaceTransform.GetScale3D().GetAbsMin();
+		}
 
 		FTransform ShapeTransform = BoundBoneTransform;
 		
diff --git a/Engine/Source/Runtime/AnimGraphRuntime/Public/BoneControllers/AnimNode_AnimDynamics.h b/Engine/Source/Runtime/AnimGraphRuntime/Public/BoneControllers/AnimNode_AnimDynamics.h
--- a/Engine/Source/Runtime/AnimGraphRuntime/Public/BoneControllers/AnimNode_AnimDynamics.h
+++ b/Engine/Source/Runtime/AnimGraphRuntime/Public/BoneControllers/AnimNode_AnimDynamics.h
@@ -685,6 +685,9 @@ private:
 	// When true, physics simulation is enabled.
 	uint8 bEnabled : 1;
 
+	// When true, input bone transform scale has changed since the last physics evaluation. Will trigger an update of physics body scales in the next physics evaluation.
+	uint8 bShouldUpdateBodiesOnScaleChange : 1;
+
 	//////////////////////////////////////////////////////////////////////////
 	// Live debug
 	//////////////////////////////////////////////////////////////////////////
