[UE] RBAN Bone Transform Scaling Fix
- Remove scaling from the BaseBoneTM at the top of the evaluate function so that any scaling in the ref bone transform for bone space simulations is not applied to the physics bodies

#Synced-CL 51591553

diff --git a/Engine/Source/Runtime/AnimGraphRuntime/Private/BoneControllers/AnimNode_RigidBody.cpp b/Engine/Source/Runtime/AnimGraphRuntime/Private/BoneControllers/AnimNode_RigidBody.cpp
--- a/Engine/Source/Runtime/AnimGraphRuntime/Private/BoneControllers/AnimNode_RigidBody.cpp
+++ b/Engine/Source/Runtime/AnimGraphRuntime/Private/BoneControllers/AnimNode_RigidBody.cpp
@@ -570,7 +570,12 @@ void FAnimNode_RigidBody::EvaluateSkeletalControl_AnyThread(FComponentSpacePoseC
 			return;
 		}
 
-		const FTransform BaseBoneTM = Output.Pose.GetComponentSpaceTransform(BaseBoneRef.GetCompactPoseIndex(BoneContainer));
+		bool bApplyBoneTransformScaleToOutputPose = bAllowPoseScalePropagation && CVarEnableRigidBodyNodePoseScalePropagation.GetValueOnAnyThread();
+
+		// Remove any scaling from BaseBoneTM if we are intending to apply the bone transform scale to the output pose. This prevents the bone transform scale
+		// being applied to the physics simulation which means that when the scale is applied to the output pose it has not effectivly been applied twice.
+		const FTransform& BaseBoneComponentSpaceTM = Output.Pose.GetComponentSpaceTransform(BaseBoneRef.GetCompactPoseIndex(BoneContainer));
+		const FTransform BaseBoneTM(BaseBoneComponentSpaceTM.GetRotation(), BaseBoneComponentSpaceTM.GetTranslation(), bApplyBoneTransformScaleToOutputPose ? FVector(1.0f) : BaseBoneComponentSpaceTM.GetScale3D());
 
 		// Initialize potential new bodies because of LOD change.
 		if (ResetSimulatedTeleportType == ETeleportType::None && bCheckForBodyTransformInit)
@@ -697,8 +702,6 @@ void FAnimNode_RigidBody::EvaluateSkeletalControl_AnyThread(FComponentSpacePoseC
 		const bool bUseDeferredSimulationTask =
 			((SimulationTiming == ESimulationTiming::Default) || bRBAN_DeferredSimulationForceDefault) ? bRBAN_DeferredSimulationDefault : (SimulationTiming == ESimulationTiming::Deferred);
 
-		bool bApplyBoneTransformScaleToOutputPose = bAllowPoseScalePropagation && CVarEnableRigidBodyNodePoseScalePropagation.GetValueOnAnyThread();
-
 		FVector SimSpaceGravity(0.f);
 
 		// Only need to tick physics if we didn't reset and we have some time to simulate
@@ -870,15 +873,10 @@ void FAnimNode_RigidBody::EvaluateSkeletalControl_AnyThread(FComponentSpacePoseC
 			}
 		}
 		
-		FTransform SimulationBaseBoneTM = BaseBoneTM;
 		TArray<FTransform> ScaledBodyTransforms;
 
 		if(bApplyBoneTransformScaleToOutputPose)
 		{
-			// BodyTM will already include any scaling in the BaseBoneTM if bone transform scale was applied above. 
-			// Remove that scaling to avoid applying it twice.
-			SimulationBaseBoneTM.SetScale3D(FVector(1.0f));
-
 			// Local record of body transforms after the input pose transform scale has been applied.
 			// Only used if bApplyBoneTransformScaleToOutputPose is true.
 			ScaledBodyTransforms.Init(FTransform::Identity, Bodies.Num());
@@ -966,7 +964,7 @@ void FAnimNode_RigidBody::EvaluateSkeletalControl_AnyThread(FComponentSpacePoseC
 					{
 						case ESimulationSpace::ComponentSpace: ComponentSpaceTM = BodyTM; break;
 						case ESimulationSpace::WorldSpace: ComponentSpaceTM = BodyTM.GetRelativeTransform(SimulationWorldSpaceTM); break;
-						case ESimulationSpace::BaseBoneSpace: ComponentSpaceTM = BodyTM * SimulationBaseBoneTM; break;
+						case ESimulationSpace::BaseBoneSpace: ComponentSpaceTM = BodyTM * BaseBoneTM; break;
 						default: ensureMsgf(false, TEXT("Unsupported Simulation Space")); ComponentSpaceTM = BodyTM;
 					}
 					
