Add support to apply additive in root space in Aim Offsets anim graph nodes

#Synced-CL 51887309

diff --git a/Engine/Source/Runtime/AnimGraphRuntime/Private/AnimNodes/AnimNode_RotationOffsetBlendSpace.cpp b/Engine/Source/Runtime/AnimGraphRuntime/Private/AnimNodes/AnimNode_RotationOffsetBlendSpace.cpp
--- a/Engine/Source/Runtime/AnimGraphRuntime/Private/AnimNodes/AnimNode_RotationOffsetBlendSpace.cpp
+++ b/Engine/Source/Runtime/AnimGraphRuntime/Private/AnimNodes/AnimNode_RotationOffsetBlendSpace.cpp
@@ -90,7 +90,15 @@ void FAnimNode_RotationOffsetBlendSpace::Evaluate_AnyThread(FPoseContext& Contex
 		// Accumulate poses together
 		FAnimationPoseData BaseAnimationPoseData(Context);
 		const FAnimationPoseData AdditiveAnimationPoseData(MeshSpaceRotationAdditivePoseContext);
-		FAnimationRuntime::AccumulateMeshSpaceRotationAdditiveToLocalPose(BaseAnimationPoseData, AdditiveAnimationPoseData, ActualAlpha);
+
+		if (bApplyAdditiveInRootSpace)
+		{
+			FAnimationRuntime::AccumulateRootSpaceRotationAdditiveToLocalPose(BaseAnimationPoseData, AdditiveAnimationPoseData, ActualAlpha);
+		}
+		else
+		{
+			FAnimationRuntime::AccumulateMeshSpaceRotationAdditiveToLocalPose(BaseAnimationPoseData, AdditiveAnimationPoseData, ActualAlpha);
+		}
 
 		// Resulting rotations are not normalized, so normalize here.
 		Context.Pose.NormalizeRotations();
diff --git a/Engine/Source/Runtime/AnimGraphRuntime/Private/AnimNodes/AnimNode_RotationOffsetBlendSpaceGraph.cpp b/Engine/Source/Runtime/AnimGraphRuntime/Private/AnimNodes/AnimNode_RotationOffsetBlendSpaceGraph.cpp
--- a/Engine/Source/Runtime/AnimGraphRuntime/Private/AnimNodes/AnimNode_RotationOffsetBlendSpaceGraph.cpp
+++ b/Engine/Source/Runtime/AnimGraphRuntime/Private/AnimNodes/AnimNode_RotationOffsetBlendSpaceGraph.cpp
@@ -55,7 +55,15 @@ void FAnimNode_RotationOffsetBlendSpaceGraph::Evaluate_AnyThread(FPoseContext& C
 		// Accumulate poses together
 		FAnimationPoseData BaseAnimationPoseData(Context);
 		const FAnimationPoseData AdditiveAnimationPoseData(MeshSpaceRotationAdditivePoseContext);
-		FAnimationRuntime::AccumulateMeshSpaceRotationAdditiveToLocalPose(BaseAnimationPoseData, AdditiveAnimationPoseData, ActualAlpha);
+
+		if (bApplyAdditiveInRootSpace)
+		{
+			FAnimationRuntime::AccumulateRootSpaceRotationAdditiveToLocalPose(BaseAnimationPoseData, AdditiveAnimationPoseData, ActualAlpha);
+		}
+		else
+		{
+			FAnimationRuntime::AccumulateMeshSpaceRotationAdditiveToLocalPose(BaseAnimationPoseData, AdditiveAnimationPoseData, ActualAlpha);
+		}
 
 		// Resulting rotations are not normalized, so normalize here.
 		Context.Pose.NormalizeRotations();
diff --git a/Engine/Source/Runtime/AnimGraphRuntime/Public/AnimNodes/AnimNode_RotationOffsetBlendSpace.h b/Engine/Source/Runtime/AnimGraphRuntime/Public/AnimNodes/AnimNode_RotationOffsetBlendSpace.h
--- a/Engine/Source/Runtime/AnimGraphRuntime/Public/AnimNodes/AnimNode_RotationOffsetBlendSpace.h
+++ b/Engine/Source/Runtime/AnimGraphRuntime/Public/AnimNodes/AnimNode_RotationOffsetBlendSpace.h
@@ -18,6 +18,10 @@ struct FAnimNode_RotationOffsetBlendSpace : public FAnimNode_BlendSpacePlayer
 	UPROPERTY(EditAnywhere, BlueprintReadWrite, Category=Links)
 	FPoseLink BasePose;
 
+	// Apply the rotation additive in root space, effectively ignoring any root bone rotation.
+	UPROPERTY(EditAnywhere, BlueprintReadWrite, Category = Settings, meta = (PinHiddenByDefault))
+	bool bApplyAdditiveInRootSpace = false;
+
 	/*
 	* Max LOD that this node is allowed to run
 	* For example if you have LODThreshold to be 2, it will run until LOD 2 (based on 0 index)
diff --git a/Engine/Source/Runtime/AnimGraphRuntime/Public/AnimNodes/AnimNode_RotationOffsetBlendSpaceGraph.h b/Engine/Source/Runtime/AnimGraphRuntime/Public/AnimNodes/AnimNode_RotationOffsetBlendSpaceGraph.h
--- a/Engine/Source/Runtime/AnimGraphRuntime/Public/AnimNodes/AnimNode_RotationOffsetBlendSpaceGraph.h
+++ b/Engine/Source/Runtime/AnimGraphRuntime/Public/AnimNodes/AnimNode_RotationOffsetBlendSpaceGraph.h
@@ -24,6 +24,10 @@ private:
 	UPROPERTY(EditAnywhere, BlueprintReadWrite, Category=Links, meta = (AllowPrivateAccess))
 	FPoseLink BasePose;
 
+	// Apply the rotation additive in root space, effectively ignoring any root bone rotation.
+	UPROPERTY(EditAnywhere, BlueprintReadWrite, Category = Settings, meta = (PinHiddenByDefault, AllowPrivateAccess))
+	bool bApplyAdditiveInRootSpace = false;
+
 	/*
 	* Max LOD that this node is allowed to run
 	* For example if you have LODThreshold to be 2, it will run until LOD 2 (based on 0 index)
