Added instrumentation for LayeredBoneBlend to display relevant info in Rewind Debugger details panel.

#rnx
#Synced-CL 52270860

diff --git a/Engine/Source/Runtime/AnimGraphRuntime/Private/AnimNodes/AnimNode_LayeredBoneBlend.cpp b/Engine/Source/Runtime/AnimGraphRuntime/Private/AnimNodes/AnimNode_LayeredBoneBlend.cpp
--- a/Engine/Source/Runtime/AnimGraphRuntime/Private/AnimNodes/AnimNode_LayeredBoneBlend.cpp
+++ b/Engine/Source/Runtime/AnimGraphRuntime/Private/AnimNodes/AnimNode_LayeredBoneBlend.cpp
@@ -230,7 +230,14 @@ void FAnimNode_LayeredBoneBlend::Update_AnyThread(const FAnimationUpdateContext&
 		}
 	}
 
+#if ANIM_TRACE_ENABLED
 	TRACE_ANIM_NODE_VALUE(Context, TEXT("Num Poses"), BlendPoses.Num());
+	TRACE_ANIM_NODE_VALUE(Context, TEXT("Has Relevant Poses"), bHasRelevantPoses);
+	for (int32 LayerIndex = 0; LayerIndex < BlendWeights.Num(); ++LayerIndex)
+	{
+		TRACE_ANIM_NODE_VALUE(Context, *FString::Printf(TEXT("Layer %d Weight"), LayerIndex), BlendWeights[LayerIndex]);
+	}
+#endif
 }
 
 void FAnimNode_LayeredBoneBlend::Evaluate_AnyThread(FPoseContext& Output)
@@ -323,11 +330,11 @@ void FAnimNode_LayeredBoneBlend::GatherDebugData(FNodeDebugData& DebugData)
 	const int NumPoses = BlendPoses.Num();
 
 	FString DebugLine = DebugData.GetNodeName(this);
-	DebugLine += FString::Printf(TEXT("(Num Poses: %i)"), NumPoses);
+	DebugLine += FString::Printf(TEXT("(Num Poses: %i, Has Relevant Poses: %s)"), NumPoses, bHasRelevantPoses ? TEXT("true") : TEXT("false"));
 	DebugData.AddDebugItem(DebugLine);
 
 	BasePose.GatherDebugData(DebugData.BranchFlow(1.f));
-	
+
 	for (int32 ChildIndex = 0; ChildIndex < NumPoses; ++ChildIndex)
 	{
 		BlendPoses[ChildIndex].GatherDebugData(DebugData.BranchFlow(BlendWeights[ChildIndex]));
