Fix montages do not mark blending out sources as inactive.

#tests editor, ABP Viewport preview
#Synced-CL 52174688

diff --git a/Engine/Source/Runtime/AnimGraphRuntime/Private/AnimNodes/AnimNode_Slot.cpp b/Engine/Source/Runtime/AnimGraphRuntime/Private/AnimNodes/AnimNode_Slot.cpp
--- a/Engine/Source/Runtime/AnimGraphRuntime/Private/AnimNodes/AnimNode_Slot.cpp
+++ b/Engine/Source/Runtime/AnimGraphRuntime/Private/AnimNodes/AnimNode_Slot.cpp
@@ -11,6 +11,16 @@
 
 #include UE_INLINE_GENERATED_CPP_BY_NAME(AnimNode_Slot)
 
+namespace MontageCVars
+{
+	bool bMarkMontageSourceAsBlendingOut = true;
+	FAutoConsoleVariableRef CVarMarkMontageSourceAsBlendingOut(
+		TEXT("a.Montage.MarkMontageSourceAsBlendingOut"),
+		bMarkMontageSourceAsBlendingOut,
+		TEXT("Set true to make montages mark their source graph as blending out when enabled")
+	);
+} // end namespace MontageCVars
+
 /////////////////////////////////////////////////////
 // FAnimNode_Slot
 
@@ -39,8 +49,10 @@ void FAnimNode_Slot::CacheBones_AnyThread(const FAnimationCacheBonesContext& Con
 void FAnimNode_Slot::Update_AnyThread(const FAnimationUpdateContext& Context)
 {
 	DECLARE_SCOPE_HIERARCHICAL_COUNTER_ANIMNODE(Update_AnyThread)
+
 	// Update weights.
-	Context.AnimInstanceProxy->GetSlotWeight(SlotName, WeightData.SlotNodeWeight, WeightData.SourceWeight, WeightData.TotalNodeWeight);
+	const float PreviousSourceWeight = WeightData.SourceWeight;
+	Context.AnimInstanceProxy->GetSlotWeight(SlotName, OUT WeightData.SlotNodeWeight, OUT WeightData.SourceWeight, OUT WeightData.TotalNodeWeight);
 
 	// Update cache in AnimInstance.
 	Context.AnimInstanceProxy->UpdateSlotNodeWeight(SlotName, WeightData.SlotNodeWeight, Context.GetFinalBlendWeight());
@@ -64,12 +76,13 @@ void FAnimNode_Slot::Update_AnyThread(const FAnimationUpdateContext& Context)
 		}
 	}
 
-
 	const bool bUpdateSource = bAlwaysUpdateSourcePose || FAnimWeight::IsRelevant(WeightData.SourceWeight);
 	if (bUpdateSource)
 	{
 		const float SourceWeight = FMath::Max(FAnimWeight::GetSmallestRelevantWeight(), WeightData.SourceWeight);
-		Source.Update(Context.FractionalWeight(SourceWeight));
+		const bool bBlendingOutSource = PreviousSourceWeight > WeightData.SourceWeight || FAnimWeight::IsFullWeight(WeightData.SlotNodeWeight);
+		const FAnimationUpdateContext SourceContext = MontageCVars::bMarkMontageSourceAsBlendingOut && bBlendingOutSource ? Context.AsInactive() : Context;
+		Source.Update(SourceContext.FractionalWeight(SourceWeight));
 	}
 
 #if ANIM_TRACE_ENABLED
