Fix for crash in Constraint AnimNode when changing setup while running

#Synced-CL 52443113

diff --git a/Engine/Source/Runtime/AnimGraphRuntime/Private/BoneControllers/AnimNode_Constraint.cpp b/Engine/Source/Runtime/AnimGraphRuntime/Private/BoneControllers/AnimNode_Constraint.cpp
--- a/Engine/Source/Runtime/AnimGraphRuntime/Private/BoneControllers/AnimNode_Constraint.cpp
+++ b/Engine/Source/Runtime/AnimGraphRuntime/Private/BoneControllers/AnimNode_Constraint.cpp
@@ -30,7 +30,7 @@ void FAnimNode_Constraint::GatherDebugData(FNodeDebugData& DebugData)
 	for (int32 ConstraintIndex = 0; ConstraintIndex < ConstraintNum; ++ConstraintIndex)
 	{
 		const FConstraint& Constraint = ConstraintSetup[ConstraintIndex];
-		const float Weight = ConstraintWeights[ConstraintIndex];
+		const float Weight = ConstraintWeights.IsValidIndex(ConstraintIndex) ? ConstraintWeights[ConstraintIndex] : 0.0f;
 		DebugLine = FString::Printf(TEXT("  Target : %s (%0.2f) "), *Constraint.TargetBone.BoneName.ToString(), Weight);
 		DebugData.AddDebugItem(DebugLine);
 	}
@@ -51,7 +51,7 @@ void FAnimNode_Constraint::EvaluateSkeletalControl_AnyThread(FComponentSpacePose
 	for (int32 ConstraintIndex = 0; ConstraintIndex < ConstraintNum; ++ConstraintIndex)
 	{
 		const FConstraint& Constraint = ConstraintSetup[ConstraintIndex];
-		const float Weight = ConstraintWeights[ConstraintIndex];
+		const float Weight = ConstraintWeights.IsValidIndex(ConstraintIndex) ? ConstraintWeights[ConstraintIndex] : 0.0f;
 
 		if (Weight > ZERO_ANIMWEIGHT_THRESH && Constraint.IsValidToEvaluate(BoneContainer) && Constraint.ConstraintDataIndex != INDEX_NONE)
 		{
@@ -80,7 +80,7 @@ void FAnimNode_Constraint::EvaluateSkeletalControl_AnyThread(FComponentSpacePose
 #if ANIM_TRACE_ENABLED
 	for (int32 ConstraintIndex = 0; ConstraintIndex < ConstraintNum; ++ConstraintIndex)
 	{
-		TRACE_ANIM_NODE_VALUE(Output, *ConstraintSetup[ConstraintIndex].TargetBone.BoneName.ToString(), ConstraintWeights[ConstraintIndex]);
+		TRACE_ANIM_NODE_VALUE(Output, *ConstraintSetup[ConstraintIndex].TargetBone.BoneName.ToString(), ConstraintWeights.IsValidIndex(ConstraintIndex) ? ConstraintWeights[ConstraintIndex] : 0.0f);
 	}
 #endif
 }
@@ -94,7 +94,7 @@ bool FAnimNode_Constraint::IsValidToEvaluate(const USkeleton* Skeleton, const FB
 	{
 		FConstraint& Constraint = ConstraintSetup[ConstraintIndex];
 		// make sure it has weight
-		if (ConstraintWeights[ConstraintIndex] > ZERO_ANIMWEIGHT_THRESH)
+		if (ConstraintWeights.IsValidIndex(ConstraintIndex) && ConstraintWeights[ConstraintIndex] > ZERO_ANIMWEIGHT_THRESH)
 		{
 			bHaveValidConstraint |= Constraint.IsValidToEvaluate(RequiredBones);
 		}
