Fix for Constraint AnimNode when initial weight is zero but in the future will be non-zero.

#Synced-CL 52573539

diff --git a/Engine/Source/Runtime/AnimGraphRuntime/Private/BoneControllers/AnimNode_Constraint.cpp b/Engine/Source/Runtime/AnimGraphRuntime/Private/BoneControllers/AnimNode_Constraint.cpp
--- a/Engine/Source/Runtime/AnimGraphRuntime/Private/BoneControllers/AnimNode_Constraint.cpp
+++ b/Engine/Source/Runtime/AnimGraphRuntime/Private/BoneControllers/AnimNode_Constraint.cpp
@@ -53,7 +53,7 @@ void FAnimNode_Constraint::EvaluateSkeletalControl_AnyThread(FComponentSpacePose
 		const FConstraint& Constraint = ConstraintSetup[ConstraintIndex];
 		const float Weight = ConstraintWeights.IsValidIndex(ConstraintIndex) ? ConstraintWeights[ConstraintIndex] : 0.0f;
 
-		if (Weight > ZERO_ANIMWEIGHT_THRESH && Constraint.IsValidToEvaluate(BoneContainer) && Constraint.ConstraintDataIndex != INDEX_NONE)
+		if (Constraint.IsValidToEvaluate(BoneContainer) && Constraint.ConstraintDataIndex != INDEX_NONE)
 		{
 			ConstraintData[Constraint.ConstraintDataIndex].Weight = Weight;
 			ConstraintData[Constraint.ConstraintDataIndex].CurrentTransform = Output.Pose.GetComponentSpaceTransform(Constraint.TargetBone.CachedCompactPoseIndex);
